hlslcc - Handle spaces in paths

diff --git a/Engine/Source/Developer/ShaderCompilerCommon/Private/ShaderCompilerCommon.cpp b/Engine/Source/Developer/ShaderCompilerCommon/Private/ShaderCompilerCommon.cpp
--- a/Engine/Source/Developer/ShaderCompilerCommon/Private/ShaderCompilerCommon.cpp
+++ b/Engine/Source/Developer/ShaderCompilerCommon/Private/ShaderCompilerCommon.cpp
@@ -164,7 +164,7 @@ FString CreateCrossCompilerBatchFileContents(const FString& ShaderFile, const FS
 	BatchFile += TEXT("\nif defined ue.hlslcc GOTO DONE\nset ue.hlslcc=");
 	BatchFile += FPaths::RootDir() / TEXT("Engine\\Source\\ThirdParty\\hlslcc\\hlslcc\\bin\\Win64\\VS2013\\hlslcc_64.exe");
 	BatchFile += TEXT("\n\n:DONE\n%ue.hlslcc% ");
-	BatchFile += FString::Printf(TEXT("%s -o=%s %s -entry=%s %s %s"), *ShaderFile, *OutputFile, *FrequencySwitch, *EntryPoint, *VersionSwitch, *ExtraArguments);
+	BatchFile += FString::Printf(TEXT("\"%s\" -o=\"%s\" %s -entry=%s %s %s"), *ShaderFile, *OutputFile, *FrequencySwitch, *EntryPoint, *VersionSwitch, *ExtraArguments);
 	BatchFile += TEXT("\npause\n");
 	return BatchFile;
 }
