UE4 - Fix shader cooker issue

diff --git a/Engine/Source/Developer/ShaderCompilerCommon/Private/ShaderCompilerCommon.cpp b/Engine/Source/Developer/ShaderCompilerCommon/Private/ShaderCompilerCommon.cpp
--- a/Engine/Source/Developer/ShaderCompilerCommon/Private/ShaderCompilerCommon.cpp
+++ b/Engine/Source/Developer/ShaderCompilerCommon/Private/ShaderCompilerCommon.cpp
@@ -211,17 +211,17 @@ int32 HlslCrossCompile(
 	switch (InShaderFrequency)
 	{
 	case SF_Vertex:
-		return HSF_VertexShader;
+		Frequency = HSF_VertexShader;
 	case SF_Pixel:
-		return HSF_PixelShader;
+		Frequency = HSF_PixelShader;
 	case SF_Compute:
-		return HSF_ComputeShader;
+		Frequency = HSF_ComputeShader;
 	case SF_Geometry:
-		return HSF_GeometryShader;
+		Frequency = HSF_GeometryShader;
 	case SF_Hull:
-		return HSF_HullShader;
+		Frequency = HSF_HullShader;
 	case SF_Domain:
-		return HSF_DomainShader;
+		Frequency = HSF_DomainShader;
 	default:
 		checkf(0, TEXT("Invalid Shader Frequency %d for CrossCompiler"), InShaderFrequency);
 		return 0;
