UE4 - refix the compiler bug

diff --git a/Engine/Source/Developer/ShaderCompilerCommon/Private/ShaderCompilerCommon.cpp b/Engine/Source/Developer/ShaderCompilerCommon/Private/ShaderCompilerCommon.cpp
--- a/Engine/Source/Developer/ShaderCompilerCommon/Private/ShaderCompilerCommon.cpp
+++ b/Engine/Source/Developer/ShaderCompilerCommon/Private/ShaderCompilerCommon.cpp
@@ -212,16 +212,22 @@ int32 HlslCrossCompile(
 	{
 	case SF_Vertex:
 		Frequency = HSF_VertexShader;
+		break;
 	case SF_Pixel:
 		Frequency = HSF_PixelShader;
+		break;
 	case SF_Compute:
 		Frequency = HSF_ComputeShader;
+		break;
 	case SF_Geometry:
 		Frequency = HSF_GeometryShader;
+		break;
 	case SF_Hull:
 		Frequency = HSF_HullShader;
+		break;
 	case SF_Domain:
 		Frequency = HSF_DomainShader;
+		break;
 	default:
 		checkf(0, TEXT("Invalid Shader Frequency %d for CrossCompiler"), InShaderFrequency);
 		return 0;
