UE4 - Hlsl Lexer (not used)

diff --git a/Engine/Source/Developer/ShaderCompilerCommon/Private/HlslLexer.cpp b/Engine/Source/Developer/ShaderCompilerCommon/Private/HlslLexer.cpp
new file mode 100644
--- /dev/null
+++ b/Engine/Source/Developer/ShaderCompilerCommon/Private/HlslLexer.cpp
@@ -0,0 +1,731 @@
+// Copyright 1998-2014 Epic Games, Inc. All Rights Reserved.
+// HlslLexer.h - Interface for scanning & tokenizing hlsl
+
+
+#include "ShaderCompilerCommon.h"
+#include "HlslLexer.h"
+
+namespace CrossCompiler
+{
+	template <typename T, size_t N>
+	char (&ArraySizeHelper(T (&array)[N]))[N];
+#define ArraySize(array) (sizeof(ArraySizeHelper(array)))
+#define MATCH_TARGET(S) S, (int32)ArraySize(S) - 1
+
+	typedef FPlatformTypes::TCHAR TCHAR;
+
+	static FORCEINLINE bool IsSpaceOrTab(TCHAR Char)
+	{
+		return Char == ' ' || Char == '\t';
+	}
+
+	static FORCEINLINE bool IsEOL(TCHAR Char)
+	{
+		return Char == '\r' || Char == '\n';
+	}
+
+	static FORCEINLINE bool IsSpaceOrTabOrEOL(TCHAR Char)
+	{
+		return IsEOL(Char) || IsSpaceOrTab(Char);
+	}
+
+	static FORCEINLINE bool IsAlpha(TCHAR Char)
+	{
+		return (Char >= 'a' && Char <= 'z') || (Char >= 'A' && Char <= 'Z');
+	}
+
+	static FORCEINLINE bool IsDigit(TCHAR Char)
+	{
+		return Char >= '0' && Char <= '9';
+	}
+
+	static FORCEINLINE bool IsAlphaOrDigit(TCHAR Char)
+	{
+		return IsAlpha(Char) || IsDigit(Char);
+	}
+
+	struct FKeywordToken
+	{
+		EHlslToken Current;
+		void* Map;
+
+		FKeywordToken() : Current(EHlslToken::Invalid), Map(nullptr) {}
+	};
+	typedef TMap<TCHAR, FKeywordToken> TCharKeywordTokenMap;
+	TCharKeywordTokenMap Keywords;
+
+/*
+	typedef TMap<TCHAR, EHlslToken> TCharTokenMap;
+	typedef TMap<TCHAR, TCharTokenMap > TCharCharTokenMap;
+	TCharCharTokenMap SymbolTokens;*/
+
+	static void InsertToken(const TCHAR* String, EHlslToken Token)
+	{
+		TCharKeywordTokenMap* Map = &Keywords;
+		while (*String)
+		{
+			FKeywordToken& KT = Map->FindOrAdd(*String);
+			++String;
+			if (!*String)
+			{
+				KT.Current = Token;
+				return;
+			}
+
+			if (!KT.Map)
+			{
+				KT.Map = new TCharKeywordTokenMap();
+			}
+
+			Map = (TCharKeywordTokenMap*)KT.Map;
+		}
+	}
+
+	static bool MatchSymbolToken(const TCHAR*& String, EHlslToken& OutToken)
+	{
+		const TCHAR* OriginalString = String;
+		FKeywordToken* Found = Keywords.Find(*String);
+		if (!Found)
+		{
+			String = OriginalString;
+			return false;
+		}
+
+		do
+		{
+			++String;
+			if (Found->Map)
+			{
+				auto* Map = (TCharKeywordTokenMap*)Found->Map;
+				FKeywordToken* NewFound = Map->Find(*String);
+				if (!NewFound)
+				{
+					if (Found->Current != EHlslToken::Invalid)
+					{
+						OutToken = Found->Current;
+						return true;
+					}
+
+					String = OriginalString;
+					return false;
+				}
+				Found = NewFound;
+			}
+			else
+			{
+				OutToken = Found->Current;
+				return true;
+			}
+		}
+		while (*String);
+
+		String = OriginalString;
+		return false;
+	}
+
+	static void InitializeTokenTables()
+	{
+		// Math
+		InsertToken(TEXT("+"), EHlslToken::Plus);
+		InsertToken(TEXT("-"), EHlslToken::Minus);
+		InsertToken(TEXT("*"), EHlslToken::Times);
+		InsertToken(TEXT("/"), EHlslToken::Div);
+		InsertToken(TEXT("%"), EHlslToken::Mod);
+		InsertToken(TEXT("("), EHlslToken::LeftParenthesis);
+		InsertToken(TEXT(")"), EHlslToken::RightParenthesis);
+
+		// Logical
+		InsertToken(TEXT("=="), EHlslToken::EqualEqual);
+		InsertToken(TEXT("!="), EHlslToken::NotEqual);
+		InsertToken(TEXT("<"), EHlslToken::Lower);
+		InsertToken(TEXT("<="), EHlslToken::LowerEqual);
+		InsertToken(TEXT(">"), EHlslToken::Greater);
+		InsertToken(TEXT(">="), EHlslToken::GreaterEqual);
+		InsertToken(TEXT("&&"), EHlslToken::AndAnd);
+		InsertToken(TEXT("||"), EHlslToken::OrOr);
+
+		// Bit
+		InsertToken(TEXT("&"), EHlslToken::And);
+		InsertToken(TEXT("|"), EHlslToken::Or);
+		InsertToken(TEXT("!"), EHlslToken::Not);
+		InsertToken(TEXT("~"), EHlslToken::Neg);
+		InsertToken(TEXT("^"), EHlslToken::Xor);
+
+		// Statements/Keywords
+		InsertToken(TEXT("="), EHlslToken::Equal);
+		InsertToken(TEXT("{"), EHlslToken::LeftBracket);
+		InsertToken(TEXT("}"), EHlslToken::RightBracket);
+		InsertToken(TEXT(";"), EHlslToken::Semicolon);
+		InsertToken(TEXT("if"), EHlslToken::If);
+		InsertToken(TEXT("for"), EHlslToken::For);
+		InsertToken(TEXT("while"), EHlslToken::While);
+		InsertToken(TEXT("do"), EHlslToken::Do);
+		InsertToken(TEXT("return"), EHlslToken::Return);
+
+		// Unary
+		InsertToken(TEXT("++"), EHlslToken::PlusPlus);
+		InsertToken(TEXT("--"), EHlslToken::MinusMinus);
+
+		// Types
+		InsertToken(TEXT("void"), EHlslToken::Void);
+		InsertToken(TEXT("const"), EHlslToken::Const);
+
+		InsertToken(TEXT("bool"), EHlslToken::Bool);
+		InsertToken(TEXT("bool1"), EHlslToken::Bool1);
+		InsertToken(TEXT("bool2"), EHlslToken::Bool2);
+		InsertToken(TEXT("bool3"), EHlslToken::Bool3);
+		InsertToken(TEXT("bool4"), EHlslToken::Bool4);
+		InsertToken(TEXT("bool1x1"), EHlslToken::Bool1x1);
+		InsertToken(TEXT("bool2x1"), EHlslToken::Bool2x1);
+		InsertToken(TEXT("bool3x1"), EHlslToken::Bool3x1);
+		InsertToken(TEXT("bool4x1"), EHlslToken::Bool4x1);
+		InsertToken(TEXT("bool1x2"), EHlslToken::Bool1x2);
+		InsertToken(TEXT("bool2x2"), EHlslToken::Bool2x2);
+		InsertToken(TEXT("bool3x2"), EHlslToken::Bool3x2);
+		InsertToken(TEXT("bool4x2"), EHlslToken::Bool4x2);
+		InsertToken(TEXT("bool1x3"), EHlslToken::Bool1x3);
+		InsertToken(TEXT("bool2x3"), EHlslToken::Bool2x3);
+		InsertToken(TEXT("bool3x3"), EHlslToken::Bool3x3);
+		InsertToken(TEXT("bool4x3"), EHlslToken::Bool4x3);
+		InsertToken(TEXT("bool1x4"), EHlslToken::Bool1x4);
+		InsertToken(TEXT("bool2x4"), EHlslToken::Bool2x4);
+		InsertToken(TEXT("bool3x4"), EHlslToken::Bool3x4);
+		InsertToken(TEXT("bool4x4"), EHlslToken::Bool4x4);
+
+		InsertToken(TEXT("int"), EHlslToken::Int);
+		InsertToken(TEXT("int1"), EHlslToken::Int1);
+		InsertToken(TEXT("int2"), EHlslToken::Int2);
+		InsertToken(TEXT("int3"), EHlslToken::Int3);
+		InsertToken(TEXT("int4"), EHlslToken::Int4);
+		InsertToken(TEXT("int1x1"), EHlslToken::Int1x1);
+		InsertToken(TEXT("int2x1"), EHlslToken::Int2x1);
+		InsertToken(TEXT("int3x1"), EHlslToken::Int3x1);
+		InsertToken(TEXT("int4x1"), EHlslToken::Int4x1);
+		InsertToken(TEXT("int1x2"), EHlslToken::Int1x2);
+		InsertToken(TEXT("int2x2"), EHlslToken::Int2x2);
+		InsertToken(TEXT("int3x2"), EHlslToken::Int3x2);
+		InsertToken(TEXT("int4x2"), EHlslToken::Int4x2);
+		InsertToken(TEXT("int1x3"), EHlslToken::Int1x3);
+		InsertToken(TEXT("int2x3"), EHlslToken::Int2x3);
+		InsertToken(TEXT("int3x3"), EHlslToken::Int3x3);
+		InsertToken(TEXT("int4x3"), EHlslToken::Int4x3);
+		InsertToken(TEXT("int1x4"), EHlslToken::Int1x4);
+		InsertToken(TEXT("int2x4"), EHlslToken::Int2x4);
+		InsertToken(TEXT("int3x4"), EHlslToken::Int3x4);
+		InsertToken(TEXT("int4x4"), EHlslToken::Int4x4);
+
+		InsertToken(TEXT("uint"), EHlslToken::Uint);
+		InsertToken(TEXT("uint1"), EHlslToken::Uint1);
+		InsertToken(TEXT("uint2"), EHlslToken::Uint2);
+		InsertToken(TEXT("uint3"), EHlslToken::Uint3);
+		InsertToken(TEXT("uint4"), EHlslToken::Uint4);
+		InsertToken(TEXT("uint1x1"), EHlslToken::Uint1x1);
+		InsertToken(TEXT("uint2x1"), EHlslToken::Uint2x1);
+		InsertToken(TEXT("uint3x1"), EHlslToken::Uint3x1);
+		InsertToken(TEXT("uint4x1"), EHlslToken::Uint4x1);
+		InsertToken(TEXT("uint1x2"), EHlslToken::Uint1x2);
+		InsertToken(TEXT("uint2x2"), EHlslToken::Uint2x2);
+		InsertToken(TEXT("uint3x2"), EHlslToken::Uint3x2);
+		InsertToken(TEXT("uint4x2"), EHlslToken::Uint4x2);
+		InsertToken(TEXT("uint1x3"), EHlslToken::Uint1x3);
+		InsertToken(TEXT("uint2x3"), EHlslToken::Uint2x3);
+		InsertToken(TEXT("uint3x3"), EHlslToken::Uint3x3);
+		InsertToken(TEXT("uint4x3"), EHlslToken::Uint4x3);
+		InsertToken(TEXT("uint1x4"), EHlslToken::Uint1x4);
+		InsertToken(TEXT("uint2x4"), EHlslToken::Uint2x4);
+		InsertToken(TEXT("uint3x4"), EHlslToken::Uint3x4);
+		InsertToken(TEXT("uint4x4"), EHlslToken::Uint4x4);
+
+		InsertToken(TEXT("half"), EHlslToken::Half);
+		InsertToken(TEXT("half1"), EHlslToken::Half1);
+		InsertToken(TEXT("half2"), EHlslToken::Half2);
+		InsertToken(TEXT("half3"), EHlslToken::Half3);
+		InsertToken(TEXT("half4"), EHlslToken::Half4);
+		InsertToken(TEXT("half1x1"), EHlslToken::Half1x1);
+		InsertToken(TEXT("half2x1"), EHlslToken::Half2x1);
+		InsertToken(TEXT("half3x1"), EHlslToken::Half3x1);
+		InsertToken(TEXT("half4x1"), EHlslToken::Half4x1);
+		InsertToken(TEXT("half1x2"), EHlslToken::Half1x2);
+		InsertToken(TEXT("half2x2"), EHlslToken::Half2x2);
+		InsertToken(TEXT("half3x2"), EHlslToken::Half3x2);
+		InsertToken(TEXT("half4x2"), EHlslToken::Half4x2);
+		InsertToken(TEXT("half1x3"), EHlslToken::Half1x3);
+		InsertToken(TEXT("half2x3"), EHlslToken::Half2x3);
+		InsertToken(TEXT("half3x3"), EHlslToken::Half3x3);
+		InsertToken(TEXT("half4x3"), EHlslToken::Half4x3);
+		InsertToken(TEXT("half1x4"), EHlslToken::Half1x4);
+		InsertToken(TEXT("half2x4"), EHlslToken::Half2x4);
+		InsertToken(TEXT("half3x4"), EHlslToken::Half3x4);
+		InsertToken(TEXT("half4x4"), EHlslToken::Half4x4);
+
+		InsertToken(TEXT("float"), EHlslToken::Float);
+		InsertToken(TEXT("float1"), EHlslToken::Float1);
+		InsertToken(TEXT("float2"), EHlslToken::Float2);
+		InsertToken(TEXT("float3"), EHlslToken::Float3);
+		InsertToken(TEXT("float4"), EHlslToken::Float4);
+		InsertToken(TEXT("float1x1"), EHlslToken::Float1x1);
+		InsertToken(TEXT("float2x1"), EHlslToken::Float2x1);
+		InsertToken(TEXT("float3x1"), EHlslToken::Float3x1);
+		InsertToken(TEXT("float4x1"), EHlslToken::Float4x1);
+		InsertToken(TEXT("float1x2"), EHlslToken::Float1x2);
+		InsertToken(TEXT("float2x2"), EHlslToken::Float2x2);
+		InsertToken(TEXT("float3x2"), EHlslToken::Float3x2);
+		InsertToken(TEXT("float4x2"), EHlslToken::Float4x2);
+		InsertToken(TEXT("float1x3"), EHlslToken::Float1x3);
+		InsertToken(TEXT("float2x3"), EHlslToken::Float2x3);
+		InsertToken(TEXT("float3x3"), EHlslToken::Float3x3);
+		InsertToken(TEXT("float4x3"), EHlslToken::Float4x3);
+		InsertToken(TEXT("float1x4"), EHlslToken::Float1x4);
+		InsertToken(TEXT("float2x4"), EHlslToken::Float2x4);
+		InsertToken(TEXT("float3x4"), EHlslToken::Float3x4);
+		InsertToken(TEXT("float4x4"), EHlslToken::Float4x4);
+		InsertToken(TEXT("Texture"), EHlslToken::Texture);
+		InsertToken(TEXT("Texture1D"), EHlslToken::Texture1D);
+		InsertToken(TEXT("Texture1DArray"), EHlslToken::Texture1DArray);
+		InsertToken(TEXT("Texture2D"), EHlslToken::Texture2D);
+		InsertToken(TEXT("Texture2DArray"), EHlslToken::Texture2DArray);
+		InsertToken(TEXT("Texture3D"), EHlslToken::Texture3D);
+		InsertToken(TEXT("TextureCube"), EHlslToken::TextureCube);
+		InsertToken(TEXT("TextureCubeArray"), EHlslToken::TextureCubeArray);
+		InsertToken(TEXT("Sampler"), EHlslToken::Sampler);
+		InsertToken(TEXT("Sampler1D"), EHlslToken::Sampler1D);
+		InsertToken(TEXT("Sampler2D"), EHlslToken::Sampler2D);
+		InsertToken(TEXT("Sampler3D"), EHlslToken::Sampler3D);
+		InsertToken(TEXT("SamplerCube"), EHlslToken::SamplerCube);
+		InsertToken(TEXT("SamplerState"), EHlslToken::SamplerState);
+		InsertToken(TEXT("SampleComparisonState"), EHlslToken::SampleComparisonState);
+		InsertToken(TEXT("Buffer"), EHlslToken::Buffer);
+		InsertToken(TEXT("AppendStructuredBuffer"), EHlslToken::AppendStructuredBuffer);
+		InsertToken(TEXT("ByteAddressBuffer"), EHlslToken::ByteAddressBuffer);
+		InsertToken(TEXT("ConsumeStructuredBuffer"), EHlslToken::ConsumeStructuredBuffer);
+		InsertToken(TEXT("InputPatch"), EHlslToken::InputPatch);
+		InsertToken(TEXT("OutputPatch"), EHlslToken::OutputPatch);
+		InsertToken(TEXT("RWBuffer"), EHlslToken::RWBuffer);
+		InsertToken(TEXT("RWByteAddressBuffer"), EHlslToken::RWByteAddressBuffer);
+		InsertToken(TEXT("RWStructuredBuffer"), EHlslToken::RWStructuredBuffer);
+		InsertToken(TEXT("RWTexture1D"), EHlslToken::RWTexture1D);
+		InsertToken(TEXT("RWTexture1DArray"), EHlslToken::RWTexture1DArray);
+		InsertToken(TEXT("RWTexture2D"), EHlslToken::RWTexture2D);
+		InsertToken(TEXT("RWTexture2DArray"), EHlslToken::RWTexture2DArray);
+		InsertToken(TEXT("RWTexture3D"), EHlslToken::RWTexture3D);
+		InsertToken(TEXT("StructuredBuffer"), EHlslToken::StructuredBuffer);
+
+		// Modifiers
+		InsertToken(TEXT("in"), EHlslToken::In);
+		InsertToken(TEXT("out"), EHlslToken::Out);
+		InsertToken(TEXT("inout"), EHlslToken::InOut);
+		InsertToken(TEXT("static"), EHlslToken::Static);
+
+		// Misc
+		InsertToken(TEXT("["), EHlslToken::LeftSquareBracket);
+		InsertToken(TEXT("]"), EHlslToken::RightSquareBracket);
+		InsertToken(TEXT("?"), EHlslToken::Question);
+		InsertToken(TEXT(":"), EHlslToken::Colon);
+		InsertToken(TEXT(","), EHlslToken::Comma);
+		InsertToken(TEXT("."), EHlslToken::Dot);
+		InsertToken(TEXT("struct"), EHlslToken::Struct);
+		InsertToken(TEXT("cbuffer"), EHlslToken::CBuffer);
+		InsertToken(TEXT("groupshared"), EHlslToken::GroupShared);
+	}
+
+	struct FTokenizer
+	{
+		FString Filename;
+		const TCHAR* Current;
+		const TCHAR* End;
+		int32 Line;
+
+		FTokenizer(const FString& InString, const FString& InFilename = TEXT("")) :
+			Filename(InFilename),
+			Current(nullptr),
+			End(nullptr),
+			Line(0)
+		{
+			if (InString.Len() > 0)
+			{
+				Current = *InString;
+				End = *InString + InString.Len();
+				Line = 1;
+			}
+
+			static bool bInitialized = false;
+			if (!bInitialized)
+			{
+				InitializeTokenTables();
+				bInitialized = true;
+			}
+		}
+
+		bool HasCharsAvailable() const
+		{
+			return Current < End;
+		}
+
+		void SkipWhitespaceInLine()
+		{
+			while (HasCharsAvailable())
+			{
+				auto Char = *Current;
+				if (!IsSpaceOrTab(Char))
+				{
+					break;
+				}
+
+				++Current;
+			}
+		}
+
+		void SkipWhitespaceAndEmptyLines()
+		{
+			while (HasCharsAvailable())
+			{
+				SkipWhitespaceInLine();
+				auto Char = Peek();
+				if (!IsEOL(Char))
+				{
+					if (Char == '/')
+					{
+						auto NextChar = Peek(1);
+						if (NextChar == '*')
+						{
+							// C Style comment, eat everything up to */
+							Current += 2;
+							bool bClosedComment = false;
+							while (HasCharsAvailable())
+							{
+								if (*Current == '*')
+								{
+									if (Peek(1) == '/')
+									{
+										bClosedComment = true;
+										Current += 2;
+										break;
+									}
+								}
+								++Current;
+							}
+							//@todo-rco: Error if no closing */ found and we got to EOL
+							//check(bClosedComment);
+						}
+						else if (NextChar == '/')
+						{
+							// C++Style comment
+							Current += 2;
+							this->SkipToNextLine();
+							continue;
+						}
+					}
+
+					break;
+				}
+				++Current;
+				++Line;
+			}
+		}
+
+		TCHAR Peek() const
+		{
+			if (HasCharsAvailable())
+			{
+				return *Current;
+			}
+
+			return 0;
+		}
+
+		TCHAR Peek(int32 Delta) const
+		{
+			check(Delta > 0);
+			if (Current + Delta < End)
+			{
+				return Current[Delta];
+			}
+
+			return 0;
+		}
+
+		void SkipToNextLine()
+		{
+			while (HasCharsAvailable())
+			{
+				auto Char = *Current;
+				++Current;
+				if (Char == '\n')
+				{
+					break;
+				}
+				check(Char != '\r');
+			}
+
+			++Line;
+		}
+
+		bool MatchString(const TCHAR* Target, int32 TargetLen)
+		{
+			if (Current + TargetLen <= End)
+			{
+				if (FCString::Strncmp(Current, Target, TargetLen) == 0)
+				{
+					Current += TargetLen;
+					return true;
+				}
+			}
+			return false;
+		}
+
+		bool PeekDigit() const
+		{
+			return (HasCharsAvailable() && IsDigit(*Current));
+		}
+
+		bool MatchUnsignedIntegerNumber(uint32& OutNum)
+		{
+			if (!PeekDigit())
+			{
+				return false;
+			}
+
+			OutNum = 0;
+			do
+			{
+				auto Next = *Current;
+				if (!IsDigit(Next))
+				{
+					break;
+				}
+				
+				OutNum = OutNum * 10 + Next - '0';
+				++Current;
+			}
+			while (HasCharsAvailable());
+			return true;
+		}
+
+		bool Match(TCHAR Char)
+		{
+			if (HasCharsAvailable() && Char == *Current)
+			{
+				++Current;
+				return true;
+			}
+
+			return false;
+		}
+
+		bool MatchFloatNumber(float& OutNum)
+		{
+			if (!PeekDigit() && Peek() != '.')
+			{
+				return false;
+			}
+
+			OutNum = FCString::Atof(Current);
+			uint32 DummyNum;
+			MatchUnsignedIntegerNumber(DummyNum);
+			if (Match('.'))
+			{
+				MatchUnsignedIntegerNumber(DummyNum);
+			}
+
+			if (Match('E') || Match('e'))
+			{
+				Match('-');
+				MatchUnsignedIntegerNumber(DummyNum);
+			}
+
+			return true;
+		}
+
+		bool MatchQuotedString(FString& OutString)
+		{
+			if (Peek() != '"')
+			{
+				return false;
+			}
+
+			OutString = TEXT("");
+			++Current;
+			while (Peek() != '"')
+			{
+				OutString += *Current;
+				//@todo-rco: Check for \"
+				//@todo-rco: Check for EOL
+				++Current;
+			}
+
+			if (Peek() == '"')
+			{
+				++Current;
+				return true;
+			}
+
+			//@todo-rco: Error!
+			return false;
+		}
+
+		bool MatchIdentifier(FString& OutIdentifier)
+		{
+			if (HasCharsAvailable())
+			{
+				auto Char = *Current;
+				if (!IsAlpha(Char) && Char != '_')
+				{
+					return false;
+				}
+
+				++Current;
+				OutIdentifier = TEXT("");
+				OutIdentifier += Char;
+				do
+				{
+					auto Char = *Current;
+					if (!IsAlphaOrDigit(Char) && Char != '_')
+					{
+						break;
+					}
+					OutIdentifier += Char;
+					++Current;
+				}
+				while (HasCharsAvailable());
+				return true;
+			}
+
+			return false;
+		}
+
+		bool MatchSymbol(EHlslToken& OutToken)
+		{
+			if (HasCharsAvailable())
+			{
+				if (MatchSymbolToken(Current, OutToken))
+				{
+					return true;
+				}
+			}
+
+			return false;
+		}
+
+		static void ProcessDirective(FTokenizer& Tokenizer)
+		{
+			check(Tokenizer.Peek() == '#');
+			if (Tokenizer.MatchString(MATCH_TARGET(TEXT("#line"))))
+			{
+				Tokenizer.SkipWhitespaceInLine();
+				uint32 Line = 0;
+				if (Tokenizer.MatchUnsignedIntegerNumber(Line))
+				{
+					Tokenizer.Line = Line;
+					Tokenizer.SkipWhitespaceInLine();
+					FString Filename;
+					if (Tokenizer.MatchQuotedString(Filename))
+					{
+						Tokenizer.Filename = Filename;
+					}
+				}
+				else
+				{
+					//@todo-rco: Warn malformed #line directive
+				}
+			}
+			else
+			{
+				//@todo-rco: Warn about unknown pragma
+			}
+
+			Tokenizer.SkipToNextLine();
+		}
+	};
+
+	struct FToken
+	{
+		EHlslToken Token;
+		FString String;
+		uint32 UnsignedInteger;
+		float Float;
+
+		explicit FToken(const FString& Identifier) : Token(EHlslToken::Identifier), String(Identifier) { }
+		explicit FToken(EHlslToken InToken) : Token(InToken) { }
+		explicit FToken(uint32 InUnsignedInteger) : Token(EHlslToken::UnsignedInteger), UnsignedInteger(InUnsignedInteger) { }
+		explicit FToken(float InFloat) : Token(EHlslToken::Float), Float(InFloat) { }
+	};
+
+	struct FHlslScanner::FHlslScannerData
+	{
+		TArray<FToken> Tokens;
+	};
+
+	FHlslScanner::FHlslScanner() :
+		Data(nullptr)
+	{
+	}
+
+	FHlslScanner::~FHlslScanner()
+	{
+		if (Data)
+		{
+			delete Data;
+		}
+	}
+
+	void FHlslScanner::Lex(const FString& String)
+	{
+		if (Data)
+		{
+			delete Data;
+		}
+
+		Data = new FHlslScannerData();
+
+		// Simple heuristic to avoid reallocating
+		Data->Tokens.Reserve(String.Len() / 4);
+
+		FTokenizer Tokenizer(String);
+		while (Tokenizer.HasCharsAvailable())
+		{
+			auto* Sanity = Tokenizer.Current;
+			Tokenizer.SkipWhitespaceAndEmptyLines();
+			if (Tokenizer.Peek() == '#')
+			{
+				FTokenizer::ProcessDirective(Tokenizer);
+			}
+			else
+			{
+				FString Identifier;
+				EHlslToken SymbolToken;
+				uint32 UnsignedInteger;
+				float FloatNumber;
+				if (Tokenizer.MatchSymbol(SymbolToken))
+				{
+					new(Data->Tokens) FToken(SymbolToken);
+				}
+				else if (Tokenizer.MatchIdentifier(Identifier))
+				{
+					new(Data->Tokens) FToken(Identifier);
+				}
+				else if (Tokenizer.MatchFloatNumber(FloatNumber))
+				{
+					new (Data->Tokens) FToken(FloatNumber);
+				}
+				else if (Tokenizer.MatchUnsignedIntegerNumber(UnsignedInteger))
+				{
+					new (Data->Tokens) FToken(UnsignedInteger);
+				}
+			}
+
+			check(Sanity != Tokenizer.Current);
+		}
+	}
+
+	void FHlslScanner::Dump()
+	{
+		if (Data)
+		{
+			for (int32 Index = 0; Index < Data->Tokens.Num(); ++Index)
+			{
+				auto& Token = Data->Tokens[Index];
+				FPlatformMisc::LowLevelOutputDebugStringf(TEXT("** %d: %d '%s'\n"), Index, Token.Token, *Token.String);
+			}
+		}
+	}
+}
diff --git a/Engine/Source/Developer/ShaderCompilerCommon/Private/HlslLexer.h b/Engine/Source/Developer/ShaderCompilerCommon/Private/HlslLexer.h
new file mode 100644
--- /dev/null
+++ b/Engine/Source/Developer/ShaderCompilerCommon/Private/HlslLexer.h
@@ -0,0 +1,227 @@
+// Copyright 1998-2014 Epic Games, Inc. All Rights Reserved.
+// HlslLexer.h - Interface for scanning & tokenizing hlsl
+
+
+namespace CrossCompiler
+{
+	enum class EHlslToken
+	{
+		// Control
+		Invalid,
+
+		// Math
+		Plus,
+		Minus,
+		Times,
+		Div,
+		Mod,
+		LeftParenthesis,
+		RightParenthesis,
+
+		// Logical
+		EqualEqual,
+		NotEqual,
+		Lower,
+		LowerEqual,
+		Greater,
+		GreaterEqual,
+		AndAnd,
+		OrOr,
+
+		// Bit
+		And,
+		Or,
+		Not,
+		Neg,
+		Xor,
+
+		// Statements
+		Equal,
+		LeftBracket,
+		RightBracket,
+		Semicolon,
+		If,
+		For,
+		While,
+		Do,
+		Return,
+
+		// Unary
+		PlusPlus,
+		MinusMinus,
+
+		// Types
+		Void,
+		Const,
+		Bool,
+		Bool1,
+		Bool2,
+		Bool3,
+		Bool4,
+		Bool1x1,
+		Bool2x1,
+		Bool3x1,
+		Bool4x1,
+		Bool1x2,
+		Bool2x2,
+		Bool3x2,
+		Bool4x2,
+		Bool1x3,
+		Bool2x3,
+		Bool3x3,
+		Bool4x3,
+		Bool1x4,
+		Bool2x4,
+		Bool3x4,
+		Bool4x4,
+		Int,
+		Int1,
+		Int2,
+		Int3,
+		Int4,
+		Int1x1,
+		Int2x1,
+		Int3x1,
+		Int4x1,
+		Int1x2,
+		Int2x2,
+		Int3x2,
+		Int4x2,
+		Int1x3,
+		Int2x3,
+		Int3x3,
+		Int4x3,
+		Int1x4,
+		Int2x4,
+		Int3x4,
+		Int4x4,
+		Uint,
+		Uint1,
+		Uint2,
+		Uint3,
+		Uint4,
+		Uint1x1,
+		Uint2x1,
+		Uint3x1,
+		Uint4x1,
+		Uint1x2,
+		Uint2x2,
+		Uint3x2,
+		Uint4x2,
+		Uint1x3,
+		Uint2x3,
+		Uint3x3,
+		Uint4x3,
+		Uint1x4,
+		Uint2x4,
+		Uint3x4,
+		Uint4x4,
+		Half,
+		Half1,
+		Half2,
+		Half3,
+		Half4,
+		Half1x1,
+		Half2x1,
+		Half3x1,
+		Half4x1,
+		Half1x2,
+		Half2x2,
+		Half3x2,
+		Half4x2,
+		Half1x3,
+		Half2x3,
+		Half3x3,
+		Half4x3,
+		Half1x4,
+		Half2x4,
+		Half3x4,
+		Half4x4,
+		Float,
+		Float1,
+		Float2,
+		Float3,
+		Float4,
+		Float1x1,
+		Float2x1,
+		Float3x1,
+		Float4x1,
+		Float1x2,
+		Float2x2,
+		Float3x2,
+		Float4x2,
+		Float1x3,
+		Float2x3,
+		Float3x3,
+		Float4x3,
+		Float1x4,
+		Float2x4,
+		Float3x4,
+		Float4x4,
+		Texture,
+		Texture1D,
+		Texture1DArray,
+		Texture2D,
+		Texture2DArray,
+		Texture3D,
+		TextureCube,
+		TextureCubeArray,
+		Sampler,
+		Sampler1D,
+		Sampler2D,
+		Sampler3D,
+		SamplerCube,
+		SamplerState,
+		SampleComparisonState,
+		Buffer,
+		AppendStructuredBuffer,
+		ByteAddressBuffer,
+		ConsumeStructuredBuffer,
+		InputPatch,
+		OutputPatch,
+		RWBuffer,
+		RWByteAddressBuffer,
+		RWStructuredBuffer,
+		RWTexture1D,
+		RWTexture1DArray,
+		RWTexture2D,
+		RWTexture2DArray,
+		RWTexture3D,
+		StructuredBuffer,
+
+		// Modifiers
+		In,
+		Out,
+		InOut,
+		Static,
+
+		// Misc
+		Identifier,
+		UnsignedInteger,
+		FloatNumber,
+
+		LeftSquareBracket,
+		RightSquareBracket,
+		Question,
+		Colon,
+		Comma,
+		Dot,
+		Struct,
+		CBuffer,
+		GroupShared,
+	};
+
+	class SHADERCOMPILERCOMMON_API FHlslScanner
+	{
+	public:
+		FHlslScanner();
+		virtual ~FHlslScanner();
+
+		void Lex(const FString& String);
+		void Dump();
+
+	private:
+		struct FHlslScannerData; 
+		FHlslScannerData* Data;
+	};
+}
diff --git a/Engine/Source/Programs/CrossCompilerTool/CrossCompilerTool.Build.cs b/Engine/Source/Programs/CrossCompilerTool/CrossCompilerTool.Build.cs
--- a/Engine/Source/Programs/CrossCompilerTool/CrossCompilerTool.Build.cs
+++ b/Engine/Source/Programs/CrossCompilerTool/CrossCompilerTool.Build.cs
@@ -11,11 +11,13 @@ public class CrossCompilerTool : ModuleRules
 		PrivateIncludePaths.Add("Runtime/Launch/Private");		// For LaunchEngineLoop.cpp include
 		PrivateIncludePaths.Add("Developer/iOS/MetalShaderFormat/Private");		// For Metal includes
 		PrivateIncludePaths.Add("Developer/ShaderFormatOpenGL/Private");		// For GLSL includes
+		PrivateIncludePaths.Add("Developer/ShaderCompilerCommon/Private");		// For Lexer includes
 
 		PrivateDependencyModuleNames.AddRange(new string []
 			{
 				"Core",
 				"Projects",
+				"ShaderCompilerCommon",
 				"MetalShaderFormat",
 				"ShaderFormatOpenGL",
 			});
