UE4 - Hlsl parser; precedence operator expression parser

diff --git a/Engine/Source/Developer/ShaderCompilerCommon/Private/HlslExpressionParser.inl b/Engine/Source/Developer/ShaderCompilerCommon/Private/HlslExpressionParser.inl
new file mode 100644
--- /dev/null
+++ b/Engine/Source/Developer/ShaderCompilerCommon/Private/HlslExpressionParser.inl
@@ -0,0 +1,714 @@
+// Copyright 1998-2014 Epic Games, Inc. All Rights Reserved.
+
+/*=============================================================================
+	HlslExpressionParser.inl - Implementation for parsing hlsl expressions.
+=============================================================================*/
+
+#include "ShaderCompilerCommon.h"
+#include "HlslParser.h"
+
+namespace CrossCompiler
+{
+	struct FSymbolScope
+	{
+		FSymbolScope* Parent;
+		TSet<FString> Symbols;
+		TArray<FSymbolScope> Children;
+
+		FSymbolScope(FSymbolScope* InParent) : Parent(InParent) {}
+
+		void Add(const FString& Type)
+		{
+			Symbols.Add(Type);
+		}
+
+		bool FindType(const FString& Type) const
+		{
+			if (Symbols.Contains(Type))
+			{
+				return true;
+			}
+
+			return Parent ? Parent->FindType(Type) : false;
+		}
+	};
+
+	struct FInfo
+	{
+		int32 Indent;
+		bool bEnabled;
+
+		FInfo(bool bInEnabled = false) : Indent(0), bEnabled(bInEnabled) {}
+
+		void Print(const FString& Message)
+		{
+			if (bEnabled)
+			{
+				FPlatformMisc::LocalPrint(*Message);
+			}
+		}
+
+		void PrintTabs()
+		{
+			if (bEnabled)
+			{
+				for (int32 Index = 0; Index < Indent; ++Index)
+				{
+					FPlatformMisc::LocalPrint(TEXT("\t"));
+				}
+			}
+		}
+
+		void PrintWithTabs(const FString& Message)
+		{
+			PrintTabs();
+			Print(Message);
+		}
+	};
+
+	struct FInfoIndentScope
+	{
+		FInfoIndentScope(FInfo& InInfo) : Info(InInfo)
+		{
+			++Info.Indent;
+		}
+
+		~FInfoIndentScope()
+		{
+			--Info.Indent;
+		}
+
+		FInfo& Info;
+	};
+
+	EParseResult ParseBasicType(EHlslToken Token)
+	{
+		switch (Token)
+		{
+		case EHlslToken::Void:
+
+		case EHlslToken::Bool:
+		case EHlslToken::Bool1:
+		case EHlslToken::Bool2:
+		case EHlslToken::Bool3:
+		case EHlslToken::Bool4:
+		case EHlslToken::Bool1x1:
+		case EHlslToken::Bool1x2:
+		case EHlslToken::Bool1x3:
+		case EHlslToken::Bool1x4:
+		case EHlslToken::Bool2x1:
+		case EHlslToken::Bool2x2:
+		case EHlslToken::Bool2x3:
+		case EHlslToken::Bool2x4:
+		case EHlslToken::Bool3x1:
+		case EHlslToken::Bool3x2:
+		case EHlslToken::Bool3x3:
+		case EHlslToken::Bool3x4:
+		case EHlslToken::Bool4x1:
+		case EHlslToken::Bool4x2:
+		case EHlslToken::Bool4x3:
+		case EHlslToken::Bool4x4:
+
+		case EHlslToken::Int:
+		case EHlslToken::Int1:
+		case EHlslToken::Int2:
+		case EHlslToken::Int3:
+		case EHlslToken::Int4:
+		case EHlslToken::Int1x1:
+		case EHlslToken::Int1x2:
+		case EHlslToken::Int1x3:
+		case EHlslToken::Int1x4:
+		case EHlslToken::Int2x1:
+		case EHlslToken::Int2x2:
+		case EHlslToken::Int2x3:
+		case EHlslToken::Int2x4:
+		case EHlslToken::Int3x1:
+		case EHlslToken::Int3x2:
+		case EHlslToken::Int3x3:
+		case EHlslToken::Int3x4:
+		case EHlslToken::Int4x1:
+		case EHlslToken::Int4x2:
+		case EHlslToken::Int4x3:
+		case EHlslToken::Int4x4:
+
+		case EHlslToken::Uint:
+		case EHlslToken::Uint1:
+		case EHlslToken::Uint2:
+		case EHlslToken::Uint3:
+		case EHlslToken::Uint4:
+		case EHlslToken::Uint1x1:
+		case EHlslToken::Uint1x2:
+		case EHlslToken::Uint1x3:
+		case EHlslToken::Uint1x4:
+		case EHlslToken::Uint2x1:
+		case EHlslToken::Uint2x2:
+		case EHlslToken::Uint2x3:
+		case EHlslToken::Uint2x4:
+		case EHlslToken::Uint3x1:
+		case EHlslToken::Uint3x2:
+		case EHlslToken::Uint3x3:
+		case EHlslToken::Uint3x4:
+		case EHlslToken::Uint4x1:
+		case EHlslToken::Uint4x2:
+		case EHlslToken::Uint4x3:
+		case EHlslToken::Uint4x4:
+
+		case EHlslToken::Half:
+		case EHlslToken::Half1:
+		case EHlslToken::Half2:
+		case EHlslToken::Half3:
+		case EHlslToken::Half4:
+		case EHlslToken::Half1x1:
+		case EHlslToken::Half1x2:
+		case EHlslToken::Half1x3:
+		case EHlslToken::Half1x4:
+		case EHlslToken::Half2x1:
+		case EHlslToken::Half2x2:
+		case EHlslToken::Half2x3:
+		case EHlslToken::Half2x4:
+		case EHlslToken::Half3x1:
+		case EHlslToken::Half3x2:
+		case EHlslToken::Half3x3:
+		case EHlslToken::Half3x4:
+		case EHlslToken::Half4x1:
+		case EHlslToken::Half4x2:
+		case EHlslToken::Half4x3:
+		case EHlslToken::Half4x4:
+
+		case EHlslToken::Float:
+		case EHlslToken::Float1:
+		case EHlslToken::Float2:
+		case EHlslToken::Float3:
+		case EHlslToken::Float4:
+		case EHlslToken::Float1x1:
+		case EHlslToken::Float1x2:
+		case EHlslToken::Float1x3:
+		case EHlslToken::Float1x4:
+		case EHlslToken::Float2x1:
+		case EHlslToken::Float2x2:
+		case EHlslToken::Float2x3:
+		case EHlslToken::Float2x4:
+		case EHlslToken::Float3x1:
+		case EHlslToken::Float3x2:
+		case EHlslToken::Float3x3:
+		case EHlslToken::Float3x4:
+		case EHlslToken::Float4x1:
+		case EHlslToken::Float4x2:
+		case EHlslToken::Float4x3:
+		case EHlslToken::Float4x4:
+			return EParseResult::Matched;
+
+		default:
+			break;
+		}
+
+		return EParseResult::NotMatched;
+	}
+
+	EParseResult ParseBasicType(FHlslScanner& Scanner)
+	{
+		auto* Token = Scanner.PeekToken();
+		if (Token)
+		{
+			if (ParseBasicType(Token->Token) == EParseResult::Matched)
+			{
+				Scanner.Advance();
+				return EParseResult::Matched;
+			}
+
+			return EParseResult::NotMatched;
+		}
+
+		return EParseResult::Error;
+	}
+
+	EParseResult ParseTextureOrBufferType(FHlslScanner& Scanner)
+	{
+		auto* Token = Scanner.PeekToken();
+		if (Token)
+		{
+			switch (Token->Token)
+			{
+			case EHlslToken::Texture:
+			case EHlslToken::Texture1D:
+			case EHlslToken::Texture1DArray:
+			case EHlslToken::Texture2D:
+			case EHlslToken::Texture2DArray:
+			case EHlslToken::Texture3D:
+			case EHlslToken::TextureCube:
+			case EHlslToken::TextureCubeArray:
+
+			case EHlslToken::Buffer:
+			case EHlslToken::AppendStructuredBuffer:
+			case EHlslToken::ByteAddressBuffer:
+			case EHlslToken::ConsumeStructuredBuffer:
+			case EHlslToken::RWBuffer:
+			case EHlslToken::RWByteAddressBuffer:
+			case EHlslToken::RWStructuredBuffer:
+			case EHlslToken::RWTexture1D:
+			case EHlslToken::RWTexture1DArray:
+			case EHlslToken::RWTexture2D:
+			case EHlslToken::RWTexture2DArray:
+			case EHlslToken::RWTexture3D:
+			case EHlslToken::StructuredBuffer:
+
+			case EHlslToken::Sampler:
+			case EHlslToken::Sampler1D:
+			case EHlslToken::Sampler2D:
+			case EHlslToken::Sampler3D:
+			case EHlslToken::SamplerCube:
+			case EHlslToken::SamplerState:
+			case EHlslToken::SampleComparisonState:
+				Scanner.Advance();
+				return EParseResult::Matched;
+
+			default:
+				break;
+			}
+
+			return EParseResult::NotMatched;
+		}
+		else
+		{
+			return EParseResult::Error;
+		}
+	}
+
+	EParseResult ParseTypeToken(const FHlslToken* Token, FSymbolScope* SymbolScope)
+	{
+		if (Token == nullptr)
+		{
+			return EParseResult::NotMatched;
+		}
+
+		if (ParseBasicType(Token->Token) == EParseResult::Matched)
+		{
+			return EParseResult::Matched;
+		}
+		else if (Token->Token == EHlslToken::Identifier && SymbolScope->FindType(Token->String))
+		{
+			return EParseResult::Matched;
+		}
+
+		return EParseResult::NotMatched;
+	}
+
+
+	EParseResult ParseTypeToken(FHlslScanner& Scanner, FSymbolScope* SymbolScope)
+	{
+		if (ParseBasicType(Scanner/*, Info*/) == EParseResult::Matched)
+		{
+			return EParseResult::Matched;
+		}
+
+		const auto* Token = Scanner.GetCurrentToken();
+		if (Token && Token->Token == EHlslToken::Identifier && SymbolScope->FindType(Token->String))
+		{
+			Scanner.Advance();
+			return EParseResult::Matched;
+		}
+
+		return EParseResult::NotMatched;
+	}
+
+	EParseResult ComputeExpr(FHlslScanner& Scanner, int32 MinPrec, FInfo& Info, FSymbolScope* SymbolScope);
+
+	EParseResult ComputeAtom(FHlslScanner& Scanner, FInfo& Info, FSymbolScope* SymbolScope)
+	{
+		// Prefix (unary)
+		bool bFoundNonUnary = false;
+		//bool bFoundTypeInitializerList = false;
+		do
+		{
+			auto* Token = Scanner.GetCurrentToken();
+			if (!Token)
+			{
+				return EParseResult::Error;
+			}
+			Info.PrintWithTabs(FString::Printf(TEXT("Atom %s\n"), *Token->String));
+
+			bool bFoundUnary = true;
+			switch (Token->Token)
+			{
+			case EHlslToken::PlusPlus:
+				//FPlatformMisc::LowLevelOutputDebugStringf(TEXT("Unary %s\n"), *Token->String);
+				Scanner.Advance();
+				break;
+
+			case EHlslToken::MinusMinus:
+				//FPlatformMisc::LowLevelOutputDebugStringf(TEXT("Unary %s\n"), *Token->String);
+				Scanner.Advance();
+				break;
+
+			case EHlslToken::Plus:
+				//FPlatformMisc::LowLevelOutputDebugStringf(TEXT("Unary %s\n"), *Token->String);
+				Scanner.Advance();
+				break;
+
+			case EHlslToken::Minus:
+				//FPlatformMisc::LowLevelOutputDebugStringf(TEXT("Unary %s\n"), *Token->String);
+				Scanner.Advance();
+				break;
+
+			case EHlslToken::Not:
+				//FPlatformMisc::LowLevelOutputDebugStringf(TEXT("Unary %s\n"), *Token->String);
+				Scanner.Advance();
+				break;
+
+			case EHlslToken::Neg:
+				//FPlatformMisc::LowLevelOutputDebugStringf(TEXT("Unary %s\n"), *Token->String);
+				Scanner.Advance();
+				break;
+
+			case EHlslToken::UnsignedIntegerConstant:
+				//FPlatformMisc::LowLevelOutputDebugStringf(TEXT("Atom Int %d\n"), Token->UnsignedInteger);
+				Scanner.Advance();
+				bFoundUnary = false;
+				bFoundNonUnary = true;
+				break;
+
+			case EHlslToken::FloatConstant:
+				//FPlatformMisc::LowLevelOutputDebugStringf(TEXT("Atom Float %f\n"), Token->Float);
+				Scanner.Advance();
+				bFoundUnary = false;
+				bFoundNonUnary = true;
+				break;
+
+			case EHlslToken::Identifier:
+				//FPlatformMisc::LowLevelOutputDebugStringf(TEXT("Atom Literal %s\n"), *Token->String);
+				Scanner.Advance();
+				bFoundUnary = false;
+				bFoundNonUnary = true;
+				break;
+			case EHlslToken::LeftParenthesis:
+			{
+				Scanner.Advance();
+
+				const auto* Peek1 = Scanner.PeekToken(0);
+				const auto* Peek2 = Scanner.PeekToken(1);
+				if (Peek1 && ParseTypeToken(Peek1, SymbolScope) == EParseResult::Matched && Peek2 && Peek2->Token == EHlslToken::RightParenthesis)
+				{
+					// Cast
+					Info.PrintWithTabs(FString::Printf(TEXT("Cast %s\n"), *Peek1->String));
+					Scanner.Advance();
+					Scanner.Advance();
+				}
+				else
+				{
+					if (ComputeExpr(Scanner, 1, Info, SymbolScope) != EParseResult::Matched)
+					{
+						Scanner.SourceError(TEXT("Expected expression!"));
+						return EParseResult::Error;
+					}
+
+					if (!Scanner.MatchToken(EHlslToken::RightParenthesis))
+					{
+						Scanner.SourceError(TEXT("Expected ')'!"));
+						return EParseResult::Error;
+					}
+
+					bFoundNonUnary = true;
+					bFoundUnary = false;
+				}
+				break;
+/*
+				//FPlatformMisc::LowLevelOutputDebugStringf(TEXT("Atom (\n"));
+
+				// Try cast expression
+				else
+				{
+					auto Result = ComputeExpr(Scanner, 1, Info, SymbolScope);
+					if (Result != EParseResult::Matched)
+					{
+						Scanner.SourceError(TEXT("Expected expression!"));
+						return EParseResult::Error;
+					}
+
+					if (!Scanner.MatchToken(EHlslToken::RightParenthesis))
+					{
+						Scanner.SourceError(TEXT("Expected ')'!"));
+						return EParseResult::Error;
+					}
+
+					bFoundNonUnary = true;
+					bFoundUnary = false;
+				}
+*/
+			}
+				break;
+			default:
+				// Grrr handle Sampler as a variable name
+				if (ParseTextureOrBufferType(Scanner) == EParseResult::Matched)
+				{
+					bFoundUnary = false;
+					bFoundNonUnary = true;
+					break;
+				}
+				// Handle float3(x,y,z)
+				else if (ParseBasicType(Scanner) == EParseResult::Matched)
+				{
+					bFoundUnary = false;
+					//bFoundTypeInitializerList = true;
+					bFoundNonUnary = true;
+					break;
+				}
+				else
+				{
+					return EParseResult::NotMatched;
+				}
+			}
+
+			if (!bFoundUnary)
+			{
+				break;
+			}
+		}
+		while (Scanner.HasMoreTokens());
+
+		if (!bFoundNonUnary)
+		{
+			Scanner.SourceError(TEXT("Expected expression!"));
+			return EParseResult::Error;
+		}
+
+		// Suffix
+		while (Scanner.HasMoreTokens())
+		{
+			auto* Token = Scanner.GetCurrentToken();
+			Info.PrintWithTabs(FString::Printf(TEXT("Atom %s\n"), *Token->String));
+			if (Scanner.MatchToken(EHlslToken::LeftSquareBracket))
+			{
+				//FPlatformMisc::LowLevelOutputDebugStringf(TEXT("Array\n"));
+				auto Result = ComputeExpr(Scanner, 1, Info, SymbolScope);
+				if (Result != EParseResult::Matched)
+				{
+					Scanner.SourceError(TEXT("Expected expression!"));
+					return EParseResult::Error;
+				}
+
+				if (!Scanner.MatchToken(EHlslToken::RightSquareBracket))
+				{
+					Scanner.SourceError(TEXT("Expected ']'!"));
+					return EParseResult::Error;
+				}
+			}
+			else if (Scanner.MatchToken(EHlslToken::Dot))
+			{
+				//FPlatformMisc::LowLevelOutputDebugStringf(TEXT("Member or swizzle\n"));
+				if (!Scanner.MatchToken(EHlslToken::Identifier))
+				{
+					Scanner.SourceError(TEXT("Expected identifier for member or swizzle!"));
+					return EParseResult::Error;
+				}
+			}
+			else if (Scanner.MatchToken(EHlslToken::LeftParenthesis))
+			{
+				Info.PrintWithTabs(TEXT("FunctionCall (\n"));
+				bool bFirst = true;
+				bool bFoundRightParenthesis = false;
+				while (Scanner.HasMoreTokens())
+				{
+					if (Scanner.MatchToken(EHlslToken::RightParenthesis))
+					{
+						bFoundRightParenthesis = true;
+						break;
+					}
+
+					if (bFirst)
+					{
+						bFirst = false;
+					}
+					else
+					{
+						if (!Scanner.MatchToken(EHlslToken::Comma))
+						{
+							Scanner.SourceError(TEXT("Expected ','!"));
+							return EParseResult::Error;
+						}
+					}
+					auto Result = ComputeExpr(Scanner, 0, Info, SymbolScope);
+					if (Result != EParseResult::Matched)
+					{
+						Scanner.SourceError(TEXT("Expected expression!"));
+						return EParseResult::Error;
+					}
+				}
+
+				if (!bFoundRightParenthesis)
+				{
+					Scanner.SourceError(TEXT("Expected ')'!"));
+					return EParseResult::Error;
+				}
+			}
+			else if (Scanner.MatchToken(EHlslToken::PlusPlus) || Scanner.MatchToken(EHlslToken::MinusMinus))
+			{
+				//FPlatformMisc::LowLevelOutputDebugStringf(TEXT("Post %s\n"), *Token->String);
+			}
+			else if (Scanner.MatchToken(EHlslToken::Question))
+			{
+				//FPlatformMisc::LowLevelOutputDebugStringf(TEXT("Ternary\n\t"));
+				if (ComputeExpr(Scanner, 0, Info, SymbolScope) != EParseResult::Matched)
+				{
+					Scanner.SourceError(TEXT("Expected expression!"));
+					return EParseResult::Error;
+				}
+				if (!Scanner.MatchToken(EHlslToken::Colon))
+				{
+					Scanner.SourceError(TEXT("Expected ':'!"));
+					return EParseResult::Error;
+				}
+				//FPlatformMisc::LowLevelOutputDebugStringf(TEXT("\t:"));
+				if (ComputeExpr(Scanner, 0, Info, SymbolScope) != EParseResult::Matched)
+				{
+					Scanner.SourceError(TEXT("Expected expression!"));
+					return EParseResult::Error;
+				}
+				break;
+			}
+			else
+			{
+				break;
+			}
+		}
+
+		return EParseResult::Matched;
+	}
+
+	int32 GetPrecedence(const FHlslToken* Token)
+	{
+		if (Token)
+		{
+			switch (Token->Token)
+			{
+			case EHlslToken::Equal:
+			case EHlslToken::PlusEqual:
+			case EHlslToken::MinusEqual:
+			case EHlslToken::TimesEqual:
+			case EHlslToken::DivEqual:
+			case EHlslToken::ModEqual:
+			case EHlslToken::GreaterGreaterEqual:
+			case EHlslToken::LowerLowerEqual:
+			case EHlslToken::AndEqual:
+			case EHlslToken::OrEqual:
+			case EHlslToken::XorEqual:
+				return 1;
+
+			case EHlslToken::Question:
+				return 2;
+
+			case EHlslToken::OrOr:
+				return 3;
+
+			case EHlslToken::AndAnd:
+				return 4;
+
+			case EHlslToken::Or:
+				return 5;
+
+			case EHlslToken::Xor:
+				return 6;
+
+			case EHlslToken::And:
+				return 7;
+
+			case EHlslToken::EqualEqual:
+			case EHlslToken::NotEqual:
+				return 8;
+
+			case EHlslToken::Lower:
+			case EHlslToken::Greater:
+			case EHlslToken::LowerEqual:
+			case EHlslToken::GreaterEqual:
+				return 9;
+
+			case EHlslToken::LowerLower:
+			case EHlslToken::GreaterGreater:
+				return 10;
+
+			case EHlslToken::Plus:
+			case EHlslToken::Minus:
+				return 11;
+
+			case EHlslToken::Times:
+			case EHlslToken::Div:
+			case EHlslToken::Mod:
+				return 12;
+
+			default:
+				break;
+			}
+		}
+
+		return -1;
+	}
+
+	bool IsBinaryOperator(const FHlslToken* Token)
+	{
+		return GetPrecedence(Token) > 0;
+	}
+
+	EParseResult ComputeExpr(FHlslScanner& Scanner, int32 MinPrec, FInfo& Info, FSymbolScope* SymbolScope)
+	{
+		auto OriginalToken = Scanner.GetCurrentTokenIndex();
+		FInfoIndentScope Scope(Info);
+		/*
+		// Precedence Climbing
+		// http://eli.thegreenplace.net/2012/08/02/parsing-expressions-by-precedence-climbing
+			compute_expr(min_prec):
+			  result = compute_atom()
+
+			  while cur token is a binary operator with precedence >= min_prec:
+				prec, assoc = precedence and associativity of current token
+				if assoc is left:
+				  next_min_prec = prec + 1
+				else:
+				  next_min_prec = prec
+				rhs = compute_expr(next_min_prec)
+				result = compute operator(result, rhs)
+
+			  return result
+		*/
+		Info.PrintWithTabs(FString::Printf(TEXT("Compute Expr %d\n"), MinPrec));
+		auto Result = ComputeAtom(Scanner, Info, SymbolScope);
+		if (Result != EParseResult::Matched)
+		{
+			return Result;
+		}
+
+		do
+		{
+			auto* Token = Scanner.GetCurrentToken();
+			int32 Precedence = GetPrecedence(Token);
+			if (!Token || !IsBinaryOperator(Token) || Precedence < MinPrec)
+			{
+				break;
+			}
+
+			Scanner.Advance();
+			auto NextMinPrec = Precedence + 1;
+			Result = ComputeExpr(Scanner, NextMinPrec, Info, SymbolScope);
+			if (Result == EParseResult::Error)
+			{
+				return EParseResult::Error;
+			}
+			else if (Result == EParseResult::NotMatched)
+			{
+				break;
+			}
+		}
+		while (Scanner.HasMoreTokens());
+
+		if (OriginalToken == Scanner.GetCurrentTokenIndex())
+		{
+			return EParseResult::NotMatched;
+		}
+
+		return EParseResult::Matched;
+	}
+
+	EParseResult ParseExpression(FHlslScanner& Scanner, FSymbolScope* SymbolScope)
+	{
+		FInfo Info(!true);
+		return ComputeExpr(Scanner, 0, Info, SymbolScope);
+	}
+}
diff --git a/Engine/Source/Developer/ShaderCompilerCommon/Private/HlslLexer.cpp b/Engine/Source/Developer/ShaderCompilerCommon/Private/HlslLexer.cpp
--- a/Engine/Source/Developer/ShaderCompilerCommon/Private/HlslLexer.cpp
+++ b/Engine/Source/Developer/ShaderCompilerCommon/Private/HlslLexer.cpp
@@ -1,5 +1,5 @@
 // Copyright 1998-2014 Epic Games, Inc. All Rights Reserved.
-// HlslLexer.h - Interface for scanning & tokenizing hlsl
+// HlslLexer.cpp - Implementation for scanning & tokenizing hlsl
 
 
 #include "ShaderCompilerCommon.h"
@@ -39,6 +39,11 @@ namespace CrossCompiler
 		return Char >= '0' && Char <= '9';
 	}
 
+	static FORCEINLINE bool IsHexDigit(TCHAR Char)
+	{
+		return IsDigit(Char) || (Char >= 'a' && Char <= 'f') || (Char >= 'A' && Char <= 'F');
+	}
+
 	static FORCEINLINE bool IsAlphaOrDigit(TCHAR Char)
 	{
 		return IsAlpha(Char) || IsDigit(Char);
@@ -54,11 +59,6 @@ namespace CrossCompiler
 	typedef TMap<TCHAR, FKeywordToken> TCharKeywordTokenMap;
 	TCharKeywordTokenMap Keywords;
 
-/*
-	typedef TMap<TCHAR, EHlslToken> TCharTokenMap;
-	typedef TMap<TCHAR, TCharTokenMap > TCharCharTokenMap;
-	TCharCharTokenMap SymbolTokens;*/
-
 	static void InsertToken(const TCHAR* String, EHlslToken Token)
 	{
 		TCharKeywordTokenMap* Map = &Keywords;
@@ -81,45 +81,76 @@ namespace CrossCompiler
 		}
 	}
 
-	static bool MatchSymbolToken(const TCHAR*& String, EHlslToken& OutToken)
+	static bool MatchSymbolToken(const TCHAR* InString, const TCHAR** OutString, EHlslToken& OutToken, FString* OutTokenString, bool bGreedy)
 	{
-		const TCHAR* OriginalString = String;
-		FKeywordToken* Found = Keywords.Find(*String);
+		const TCHAR* OriginalString = InString;
+		FKeywordToken* Found = Keywords.Find(*InString);
+
+		if (OutString)
+		{
+			*OutString = OriginalString;
+		}
+
 		if (!Found)
 		{
-			String = OriginalString;
 			return false;
 		}
 
 		do
 		{
-			++String;
+			++InString;
 			if (Found->Map)
 			{
 				auto* Map = (TCharKeywordTokenMap*)Found->Map;
-				FKeywordToken* NewFound = Map->Find(*String);
+				FKeywordToken* NewFound = Map->Find(*InString);
 				if (!NewFound)
 				{
 					if (Found->Current != EHlslToken::Invalid)
 					{
-						OutToken = Found->Current;
-						return true;
+						// Don't early out on a partial match (e.g., Texture1DSample should not be 2 tokens)
+						if (!bGreedy || !*InString)
+						{
+							OutToken = Found->Current;
+							if (OutTokenString)
+							{
+								*OutTokenString = TEXT("");
+								OutTokenString->AppendChars(OriginalString, InString - OriginalString);
+							}
+
+							if (OutString)
+							{
+								*OutString = InString;
+							}
+							return true;
+						}
 					}
 
-					String = OriginalString;
 					return false;
 				}
 				Found = NewFound;
 			}
+			else if (bGreedy && *InString)
+			{
+				break;
+			}
 			else
 			{
 				OutToken = Found->Current;
+				if (OutTokenString)
+				{
+					*OutTokenString = TEXT("");
+					OutTokenString->AppendChars(OriginalString, InString - OriginalString);
+				}
+
+				if (OutString)
+				{
+					*OutString = InString;
+				}
 				return true;
 			}
 		}
-		while (*String);
+		while (*InString);
 
-		String = OriginalString;
 		return false;
 	}
 
@@ -127,10 +158,15 @@ namespace CrossCompiler
 	{
 		// Math
 		InsertToken(TEXT("+"), EHlslToken::Plus);
+		InsertToken(TEXT("+="), EHlslToken::PlusEqual);
 		InsertToken(TEXT("-"), EHlslToken::Minus);
+		InsertToken(TEXT("-="), EHlslToken::MinusEqual);
 		InsertToken(TEXT("*"), EHlslToken::Times);
+		InsertToken(TEXT("*="), EHlslToken::TimesEqual);
 		InsertToken(TEXT("/"), EHlslToken::Div);
+		InsertToken(TEXT("/="), EHlslToken::DivEqual);
 		InsertToken(TEXT("%"), EHlslToken::Mod);
+		InsertToken(TEXT("%="), EHlslToken::ModEqual);
 		InsertToken(TEXT("("), EHlslToken::LeftParenthesis);
 		InsertToken(TEXT(")"), EHlslToken::RightParenthesis);
 
@@ -145,22 +181,35 @@ namespace CrossCompiler
 		InsertToken(TEXT("||"), EHlslToken::OrOr);
 
 		// Bit
+		InsertToken(TEXT("<<"), EHlslToken::LowerLower);
+		InsertToken(TEXT("<<="), EHlslToken::LowerLowerEqual);
+		InsertToken(TEXT(">>"), EHlslToken::GreaterGreater);
+		InsertToken(TEXT(">>="), EHlslToken::GreaterGreaterEqual);
 		InsertToken(TEXT("&"), EHlslToken::And);
+		InsertToken(TEXT("&="), EHlslToken::And);
 		InsertToken(TEXT("|"), EHlslToken::Or);
+		InsertToken(TEXT("|="), EHlslToken::OrEqual);
+		InsertToken(TEXT("^"), EHlslToken::Xor);
+		InsertToken(TEXT("^="), EHlslToken::XorEqual);
 		InsertToken(TEXT("!"), EHlslToken::Not);
 		InsertToken(TEXT("~"), EHlslToken::Neg);
-		InsertToken(TEXT("^"), EHlslToken::Xor);
 
 		// Statements/Keywords
 		InsertToken(TEXT("="), EHlslToken::Equal);
-		InsertToken(TEXT("{"), EHlslToken::LeftBracket);
-		InsertToken(TEXT("}"), EHlslToken::RightBracket);
+		InsertToken(TEXT("{"), EHlslToken::LeftBrace);
+		InsertToken(TEXT("}"), EHlslToken::RightBrace);
 		InsertToken(TEXT(";"), EHlslToken::Semicolon);
 		InsertToken(TEXT("if"), EHlslToken::If);
+		InsertToken(TEXT("else"), EHlslToken::Else);
 		InsertToken(TEXT("for"), EHlslToken::For);
 		InsertToken(TEXT("while"), EHlslToken::While);
 		InsertToken(TEXT("do"), EHlslToken::Do);
 		InsertToken(TEXT("return"), EHlslToken::Return);
+		InsertToken(TEXT("switch"), EHlslToken::Switch);
+		InsertToken(TEXT("case"), EHlslToken::Case);
+		InsertToken(TEXT("break"), EHlslToken::Break);
+		InsertToken(TEXT("default"), EHlslToken::Default);
+		InsertToken(TEXT("goto"), EHlslToken::Goto);
 
 		// Unary
 		InsertToken(TEXT("++"), EHlslToken::PlusPlus);
@@ -279,6 +328,7 @@ namespace CrossCompiler
 		InsertToken(TEXT("float2x4"), EHlslToken::Float2x4);
 		InsertToken(TEXT("float3x4"), EHlslToken::Float3x4);
 		InsertToken(TEXT("float4x4"), EHlslToken::Float4x4);
+
 		InsertToken(TEXT("Texture"), EHlslToken::Texture);
 		InsertToken(TEXT("Texture1D"), EHlslToken::Texture1D);
 		InsertToken(TEXT("Texture1DArray"), EHlslToken::Texture1DArray);
@@ -287,6 +337,7 @@ namespace CrossCompiler
 		InsertToken(TEXT("Texture3D"), EHlslToken::Texture3D);
 		InsertToken(TEXT("TextureCube"), EHlslToken::TextureCube);
 		InsertToken(TEXT("TextureCubeArray"), EHlslToken::TextureCubeArray);
+
 		InsertToken(TEXT("Sampler"), EHlslToken::Sampler);
 		InsertToken(TEXT("Sampler1D"), EHlslToken::Sampler1D);
 		InsertToken(TEXT("Sampler2D"), EHlslToken::Sampler2D);
@@ -294,12 +345,11 @@ namespace CrossCompiler
 		InsertToken(TEXT("SamplerCube"), EHlslToken::SamplerCube);
 		InsertToken(TEXT("SamplerState"), EHlslToken::SamplerState);
 		InsertToken(TEXT("SampleComparisonState"), EHlslToken::SampleComparisonState);
+
 		InsertToken(TEXT("Buffer"), EHlslToken::Buffer);
 		InsertToken(TEXT("AppendStructuredBuffer"), EHlslToken::AppendStructuredBuffer);
 		InsertToken(TEXT("ByteAddressBuffer"), EHlslToken::ByteAddressBuffer);
 		InsertToken(TEXT("ConsumeStructuredBuffer"), EHlslToken::ConsumeStructuredBuffer);
-		InsertToken(TEXT("InputPatch"), EHlslToken::InputPatch);
-		InsertToken(TEXT("OutputPatch"), EHlslToken::OutputPatch);
 		InsertToken(TEXT("RWBuffer"), EHlslToken::RWBuffer);
 		InsertToken(TEXT("RWByteAddressBuffer"), EHlslToken::RWByteAddressBuffer);
 		InsertToken(TEXT("RWStructuredBuffer"), EHlslToken::RWStructuredBuffer);
@@ -309,6 +359,8 @@ namespace CrossCompiler
 		InsertToken(TEXT("RWTexture2DArray"), EHlslToken::RWTexture2DArray);
 		InsertToken(TEXT("RWTexture3D"), EHlslToken::RWTexture3D);
 		InsertToken(TEXT("StructuredBuffer"), EHlslToken::StructuredBuffer);
+		InsertToken(TEXT("InputPatch"), EHlslToken::InputPatch);
+		InsertToken(TEXT("OutputPatch"), EHlslToken::OutputPatch);
 
 		// Modifiers
 		InsertToken(TEXT("in"), EHlslToken::In);
@@ -333,12 +385,14 @@ namespace CrossCompiler
 		FString Filename;
 		const TCHAR* Current;
 		const TCHAR* End;
+		const TCHAR* CurrentLineStart;
 		int32 Line;
 
 		FTokenizer(const FString& InString, const FString& InFilename = TEXT("")) :
 			Filename(InFilename),
 			Current(nullptr),
 			End(nullptr),
+			CurrentLineStart(nullptr),
 			Line(0)
 		{
 			if (InString.Len() > 0)
@@ -346,6 +400,7 @@ namespace CrossCompiler
 				Current = *InString;
 				End = *InString + InString.Len();
 				Line = 1;
+				CurrentLineStart = Current;
 			}
 
 			static bool bInitialized = false;
@@ -365,7 +420,7 @@ namespace CrossCompiler
 		{
 			while (HasCharsAvailable())
 			{
-				auto Char = *Current;
+				auto Char = Peek();
 				if (!IsSpaceOrTab(Char))
 				{
 					break;
@@ -381,45 +436,55 @@ namespace CrossCompiler
 			{
 				SkipWhitespaceInLine();
 				auto Char = Peek();
-				if (!IsEOL(Char))
+				if (IsEOL(Char))
+				{
+					SkipToNextLine();
+				}
+				else
 				{
-					if (Char == '/')
+					auto NextChar = Peek(1);
+					if (Char == '/' && NextChar == '/')
+					{
+						// C++ comment
+						Current += 2;
+						this->SkipToNextLine();
+						continue;
+					}
+					else if (Char == '/' && NextChar == '*')
 					{
-						auto NextChar = Peek(1);
-						if (NextChar == '*')
+						// C Style comment, eat everything up to * /
+						Current += 2;
+						bool bClosedComment = false;
+						while (HasCharsAvailable())
 						{
-							// C Style comment, eat everything up to */
-							Current += 2;
-							bool bClosedComment = false;
-							while (HasCharsAvailable())
+							if (Peek() == '*')
 							{
-								if (*Current == '*')
+								if (Peek(1) == '/')
 								{
-									if (Peek(1) == '/')
-									{
-										bClosedComment = true;
-										Current += 2;
-										break;
-									}
+									bClosedComment = true;
+									Current += 2;
+									break;
 								}
-								++Current;
+
 							}
-							//@todo-rco: Error if no closing */ found and we got to EOL
-							//check(bClosedComment);
-						}
-						else if (NextChar == '/')
-						{
-							// C++Style comment
-							Current += 2;
-							this->SkipToNextLine();
-							continue;
+							else if (Peek() == '\n')
+							{
+								SkipToNextLine();
+
+								// Don't increment current!
+								continue;
+							}
+
+							++Current;
 						}
+						//@todo-rco: Error if no closing * / found and we got to EOL
+						//check(bClosedComment);
+					}
+					else
+					{
+						break;
 					}
-
-					break;
 				}
-				++Current;
-				++Line;
 			}
 		}
 
@@ -448,7 +513,7 @@ namespace CrossCompiler
 		{
 			while (HasCharsAvailable())
 			{
-				auto Char = *Current;
+				auto Char = Peek();
 				++Current;
 				if (Char == '\n')
 				{
@@ -458,6 +523,7 @@ namespace CrossCompiler
 			}
 
 			++Line;
+			CurrentLineStart = Current;
 		}
 
 		bool MatchString(const TCHAR* Target, int32 TargetLen)
@@ -475,35 +541,23 @@ namespace CrossCompiler
 
 		bool PeekDigit() const
 		{
-			return (HasCharsAvailable() && IsDigit(*Current));
+			return IsDigit(Peek());
 		}
 
-		bool MatchUnsignedIntegerNumber(uint32& OutNum)
+		bool MatchAndSkipDigits()
 		{
-			if (!PeekDigit())
+			auto* Original = Current;
+			while (PeekDigit())
 			{
-				return false;
-			}
-
-			OutNum = 0;
-			do
-			{
-				auto Next = *Current;
-				if (!IsDigit(Next))
-				{
-					break;
-				}
-				
-				OutNum = OutNum * 10 + Next - '0';
 				++Current;
 			}
-			while (HasCharsAvailable());
-			return true;
+
+			return Original != Current;
 		}
 
 		bool Match(TCHAR Char)
 		{
-			if (HasCharsAvailable() && Char == *Current)
+			if (Char == Peek())
 			{
 				++Current;
 				return true;
@@ -514,52 +568,106 @@ namespace CrossCompiler
 
 		bool MatchFloatNumber(float& OutNum)
 		{
-			if (!PeekDigit() && Peek() != '.')
+			auto* Original = Current;
+			TCHAR Char = Peek();
+
+			// \.[0-9]+([eE][+-]?[0-9]+)?[fF]?			-> Dot Digits+ Exp? F?
+			// [0-9]+\.([eE][+-]?[0-9]+)?[fF]?			-> Digits+ Dot Exp? F?
+			// [0-9]+\.[0-9]+([eE][+-]?[0-9]+)?[fF]?	-> Digits+ Dot Digits+ Exp? F?
+			// [0-9]+[eE][+-]?[0-9]+[fF]?				-> Digits+ Exp F?
+			// [0-9]+[fF]								-> Digits+ F
+			if (!IsDigit(Char) && Char != '.')
 			{
 				return false;
 			}
 
-			OutNum = FCString::Atof(Current);
-			uint32 DummyNum;
-			MatchUnsignedIntegerNumber(DummyNum);
-			if (Match('.'))
+			bool bExpOptional = false;
+			if (Match('.') && MatchAndSkipDigits())
 			{
-				MatchUnsignedIntegerNumber(DummyNum);
+				bExpOptional = true;
 			}
+			else if (MatchAndSkipDigits())
+			{
+				if (Match('.'))
+				{
+					bExpOptional = true;
+					MatchAndSkipDigits();
+				}
+				else
+				{
+					if (Match('f') || Match('F'))
+					{
+						goto Done;
+					}
 
-			if (Match('E') || Match('e'))
+					bExpOptional = false;
+				}
+			}
+			else
 			{
-				Match('-');
-				MatchUnsignedIntegerNumber(DummyNum);
+				goto NotFloat;
 			}
 
+			// Exponent [eE][+-]?[0-9]+
+			bool bExponentFound = false;
+			if (Match('e') || Match('E'))
+			{
+				Char = Peek();
+				if (Char == '+' || Char == '-')
+				{
+					++Current;
+				}
+
+				if (MatchAndSkipDigits())
+				{
+					bExponentFound = true;
+				}
+			}
+
+			if (!bExponentFound && !bExpOptional)
+			{
+				goto NotFloat;
+			}
+
+			// [fF]
+			Char = Peek();
+			if (Char == 'F' || Char == 'f')
+			{
+				++Current;
+			}
+
+		Done:
+			OutNum = FCString::Atof(Original);
 			return true;
+
+		NotFloat:
+			Current = Original;
+			return false;
 		}
 
 		bool MatchQuotedString(FString& OutString)
 		{
-			if (Peek() != '"')
+			if (!Match('"'))
 			{
 				return false;
 			}
 
 			OutString = TEXT("");
-			++Current;
 			while (Peek() != '"')
 			{
-				OutString += *Current;
+				OutString += Peek();
 				//@todo-rco: Check for \"
 				//@todo-rco: Check for EOL
 				++Current;
 			}
 
-			if (Peek() == '"')
+			if (Match('"'))
 			{
-				++Current;
 				return true;
 			}
 
 			//@todo-rco: Error!
+			check(0);
 			return false;
 		}
 
@@ -567,7 +675,7 @@ namespace CrossCompiler
 		{
 			if (HasCharsAvailable())
 			{
-				auto Char = *Current;
+				auto Char = Peek();
 				if (!IsAlpha(Char) && Char != '_')
 				{
 					return false;
@@ -578,7 +686,7 @@ namespace CrossCompiler
 				OutIdentifier += Char;
 				do
 				{
-					Char = *Current;
+					Char = Peek();
 					if (!IsAlphaOrDigit(Char) && Char != '_')
 					{
 						break;
@@ -593,11 +701,11 @@ namespace CrossCompiler
 			return false;
 		}
 
-		bool MatchSymbol(EHlslToken& OutToken)
+		bool MatchSymbol(EHlslToken& OutToken, FString& OutTokenString)
 		{
 			if (HasCharsAvailable())
 			{
-				if (MatchSymbolToken(Current, OutToken))
+				if (MatchSymbolToken(Current, &Current, OutToken, &OutTokenString, false))
 				{
 					return true;
 				}
@@ -613,9 +721,9 @@ namespace CrossCompiler
 			{
 				Tokenizer.SkipWhitespaceInLine();
 				uint32 Line = 0;
-				if (Tokenizer.MatchUnsignedIntegerNumber(Line))
+				if (Tokenizer.RuleInteger(Line))
 				{
-					Tokenizer.Line = Line;
+					Tokenizer.Line = Line - 1;
 					Tokenizer.SkipWhitespaceInLine();
 					FString Filename;
 					if (Tokenizer.MatchQuotedString(Filename))
@@ -626,59 +734,164 @@ namespace CrossCompiler
 				else
 				{
 					//@todo-rco: Warn malformed #line directive
+					check(0);
 				}
 			}
 			else
 			{
 				//@todo-rco: Warn about unknown pragma
+				check(0);
 			}
 
 			Tokenizer.SkipToNextLine();
 		}
-	};
 
-	struct FToken
-	{
-		EHlslToken Token;
-		FString String;
-		uint32 UnsignedInteger;
-		float Float;
-
-		explicit FToken(const FString& Identifier) : Token(EHlslToken::Identifier), String(Identifier) { }
-		explicit FToken(EHlslToken InToken) : Token(InToken) { }
-		explicit FToken(uint32 InUnsignedInteger) : Token(EHlslToken::UnsignedInteger), UnsignedInteger(InUnsignedInteger) { }
-		explicit FToken(float InFloat) : Token(EHlslToken::Float), Float(InFloat) { }
-	};
+		bool RuleDecimalInteger(uint32& OutValue)
+		{
+			// [1-9][0-9]*
+			auto Char = Peek();
+			{
+				if (Char < '1' || Char > '9')
+				{
+					return false;
+				}
 
-	struct FHlslScanner::FHlslScannerData
-	{
-		TArray<FToken> Tokens;
+				++Current;
+				OutValue = Char - '0';
+			}
+
+			while (HasCharsAvailable())
+			{
+				Char = Peek();
+				if (!IsDigit(Char))
+				{
+					break;
+				}
+				OutValue = OutValue * 10 + Char - '0';
+				++Current;
+			}
+
+			return true;
+		}
+
+		bool RuleOctalInteger(uint32& OutValue)
+		{
+			// 0[0-7]*
+			auto Char = Peek();
+			if (Char != '0')
+			{
+				return false;
+			}
+
+			OutValue = 0;
+			++Current;
+			while (HasCharsAvailable())
+			{
+				Char = Peek();
+				if (Char >= '0' && Char <= '7')
+				{
+					OutValue = OutValue * 8 + Char - '0';
+				}
+				else
+				{
+					break;
+				}
+				++Current;
+			}
+
+			return true;
+		}
+
+		bool RuleHexadecimalInteger(uint32& OutValue)
+		{
+			// 0[xX][0-9a-zA-Z]+
+			auto Char = Peek();
+			auto Char1 = Peek(1);
+			auto Char2 = Peek(2);
+			if (Char == '0' && (Char1 == 'x' || Char1 == 'X') && IsHexDigit(Char2))
+			{
+				Current += 2;
+				OutValue = 0;
+				do
+				{
+					Char = Peek();
+					if (IsDigit(Char))
+					{
+						OutValue = OutValue * 16 + Char - '0';
+					}
+					else if (Char >= 'a' && Char <= 'f')
+					{
+						OutValue = OutValue * 16 + Char - 'a' + 10;
+					}
+					else if (Char >= 'A' && Char <= 'F')
+					{
+						OutValue = OutValue * 16 + Char - 'A' + 10;
+					}
+					else
+					{
+						break;
+					}
+					++Current;
+				}
+				while (HasCharsAvailable());
+
+				return true;
+			}
+
+			return false;
+		}
+
+		bool RuleInteger(uint32& OutValue)
+		{
+			return RuleDecimalInteger(OutValue) || RuleHexadecimalInteger(OutValue) || RuleOctalInteger(OutValue);
+		}
+
+		bool MatchLiteralInteger(uint32& OutValue)
+		{
+			if (RuleInteger(OutValue))
+			{
+				auto Char = Peek();
+				if (Char == 'u' || Char == 'U')
+				{
+					++Current;
+				}
+
+				return true;
+			}
+
+			return false;
+		}
 	};
 
 	FHlslScanner::FHlslScanner() :
-		Data(nullptr)
+		CurrentToken(0)
 	{
 	}
 
 	FHlslScanner::~FHlslScanner()
 	{
-		if (Data)
-		{
-			delete Data;
-		}
 	}
 
-	void FHlslScanner::Lex(const FString& String)
+	inline void FHlslScanner::AddToken(const FHlslToken& Token, const FTokenizer& Tokenizer)
 	{
-		if (Data)
-		{
-			delete Data;
-		}
+		int32 TokenIndex = Tokens.Add(Token);
+		Tokens[TokenIndex].SourceFilename = &SourceFilenames.Last();
+		Tokens[TokenIndex].SourceLine = Tokenizer.Line;
+		Tokens[TokenIndex].SourceColumn = (int32)(Tokenizer.Current - Tokenizer.CurrentLineStart) + 1;
+	}
 
-		Data = new FHlslScannerData();
+	void FHlslScanner::Clear()
+	{
+		Tokens.Empty();
+		new (SourceFilenames) FString(TEXT(""));
+	}
+
+	bool FHlslScanner::Lex(const FString& String)
+	{
+		Clear();
 
 		// Simple heuristic to avoid reallocating
-		Data->Tokens.Reserve(String.Len() / 4);
+		Tokens.Reserve(String.Len() / 4);
 
 		FTokenizer Tokenizer(String);
 		while (Tokenizer.HasCharsAvailable())
@@ -688,6 +901,10 @@ namespace CrossCompiler
 			if (Tokenizer.Peek() == '#')
 			{
 				FTokenizer::ProcessDirective(Tokenizer);
+				if (Tokenizer.Filename != SourceFilenames.Last())
+				{
+					new(SourceFilenames) FString(Tokenizer.Filename);
+				}
 			}
 			else
 			{
@@ -695,37 +912,148 @@ namespace CrossCompiler
 				EHlslToken SymbolToken;
 				uint32 UnsignedInteger;
 				float FloatNumber;
-				if (Tokenizer.MatchSymbol(SymbolToken))
+				if (Tokenizer.MatchFloatNumber(FloatNumber))
+				{
+					AddToken(FHlslToken(FloatNumber), Tokenizer);
+				}
+				else if (Tokenizer.MatchLiteralInteger(UnsignedInteger))
 				{
-					new(Data->Tokens) FToken(SymbolToken);
+					AddToken(FHlslToken(UnsignedInteger), Tokenizer);
 				}
 				else if (Tokenizer.MatchIdentifier(Identifier))
 				{
-					new(Data->Tokens) FToken(Identifier);
+					if (!FCString::Strcmp(*Identifier, TEXT("true")))
+					{
+						AddToken(FHlslToken(true), Tokenizer);
+					}
+					else if (!FCString::Strcmp(*Identifier, TEXT("false")))
+					{
+						AddToken(FHlslToken(false), Tokenizer);
+					}
+					else if (MatchSymbolToken(*Identifier, nullptr, SymbolToken, nullptr, true))
+					{
+						AddToken(FHlslToken(SymbolToken, Identifier), Tokenizer);
+					}
+					else
+					{
+						AddToken(FHlslToken(Identifier), Tokenizer);
+					}
 				}
-				else if (Tokenizer.MatchFloatNumber(FloatNumber))
+				else if (Tokenizer.MatchSymbol(SymbolToken, Identifier))
 				{
-					new (Data->Tokens) FToken(FloatNumber);
+					AddToken(FHlslToken(SymbolToken, Identifier), Tokenizer);
 				}
-				else if (Tokenizer.MatchUnsignedIntegerNumber(UnsignedInteger))
+				else if (Tokenizer.HasCharsAvailable())
 				{
-					new (Data->Tokens) FToken(UnsignedInteger);
+					//@todo-rco: Unknown token!
+					if (Tokenizer.Filename.Len() > 0)
+					{
+						FPlatformMisc::LowLevelOutputDebugStringf(TEXT("Unknown token at line %d, file '%s'!"), Tokenizer.Line, *Tokenizer.Filename);
+					}
+					else
+					{
+						FPlatformMisc::LowLevelOutputDebugStringf(TEXT("Unknown token at line %d!"), Tokenizer.Line);
+					}
+					return false;
 				}
 			}
 
 			check(Sanity != Tokenizer.Current);
 		}
+
+		return true;
 	}
 
 	void FHlslScanner::Dump()
 	{
-		if (Data)
+		for (int32 Index = 0; Index < Tokens.Num(); ++Index)
 		{
-			for (int32 Index = 0; Index < Data->Tokens.Num(); ++Index)
+			auto& Token = Tokens[Index];
+			switch (Token.Token)
 			{
-				auto& Token = Data->Tokens[Index];
+			case EHlslToken::UnsignedIntegerConstant:
+				FPlatformMisc::LowLevelOutputDebugStringf(TEXT("** %d: UnsignedIntegerConstant '%d'\n"), Index, Token.UnsignedInteger);
+				break;
+
+			case EHlslToken::FloatConstant:
+				FPlatformMisc::LowLevelOutputDebugStringf(TEXT("** %d: FloatConstant '%f'\n"), Index, Token.Float);
+				break;
+
+			default:
 				FPlatformMisc::LowLevelOutputDebugStringf(TEXT("** %d: %d '%s'\n"), Index, Token.Token, *Token.String);
+				break;
 			}
 		}
 	}
+
+	bool FHlslScanner::MatchToken(EHlslToken InToken)
+	{
+		const auto* Token = GetCurrentToken();
+		if (Token)
+		{
+			if (Token->Token == InToken)
+			{
+				++CurrentToken;
+				return true;
+			}
+		}
+
+		return false;
+	}
+
+	const FHlslToken* FHlslScanner::PeekToken(uint32 LookAhead /*= 0*/) const
+	{
+		if (CurrentToken + LookAhead < (uint32)Tokens.Num())
+		{
+			return &Tokens[CurrentToken + LookAhead];
+		}
+
+		return nullptr;
+	}
+
+	bool FHlslScanner::HasMoreTokens() const
+	{
+		return CurrentToken < (uint32)Tokens.Num();
+	}
+
+	const FHlslToken* FHlslScanner::GetCurrentToken() const
+	{
+		if (CurrentToken < (uint32)Tokens.Num())
+		{
+			return &Tokens[CurrentToken];
+		}
+
+		return nullptr;
+	}
+
+	const FHlslToken* FHlslScanner::GetCurrentTokenAndAdvance()
+	{
+		if (CurrentToken < (uint32)Tokens.Num())
+		{
+			auto* Return = &Tokens[CurrentToken];
+			Advance();
+		}
+
+		return nullptr;
+	}
+
+	void FHlslScanner::SetCurrentTokenIndex(uint32 NewToken)
+	{
+		check(NewToken <= (uint32)Tokens.Num());
+		CurrentToken = NewToken;
+	}
+
+	void FHlslScanner::SourceError(const FString& Error)
+	{
+		if (CurrentToken < (uint32)Tokens.Num())
+		{
+			const auto& Token = Tokens[CurrentToken];
+			check(Token.SourceFilename);
+			FPlatformMisc::LowLevelOutputDebugStringf(TEXT("%s(%d): (%d) %s\n"), **Token.SourceFilename, Token.SourceLine, Token.SourceColumn, *Error);
+		}
+		else
+		{
+			FPlatformMisc::LowLevelOutputDebugStringf(TEXT("%s\n"), *Error);
+		}
+	}
 }
diff --git a/Engine/Source/Developer/ShaderCompilerCommon/Private/HlslLexer.h b/Engine/Source/Developer/ShaderCompilerCommon/Private/HlslLexer.h
--- a/Engine/Source/Developer/ShaderCompilerCommon/Private/HlslLexer.h
+++ b/Engine/Source/Developer/ShaderCompilerCommon/Private/HlslLexer.h
@@ -1,6 +1,10 @@
 // Copyright 1998-2014 Epic Games, Inc. All Rights Reserved.
-// HlslLexer.h - Interface for scanning & tokenizing hlsl
 
+/*=============================================================================
+	HlslLexer.h - Interface for scanning & tokenizing hlsl.
+=============================================================================*/
+
+#pragma once
 
 namespace CrossCompiler
 {
@@ -11,10 +15,15 @@ namespace CrossCompiler
 
 		// Math
 		Plus,
+		PlusEqual,
 		Minus,
+		MinusEqual,
 		Times,
+		TimesEqual,
 		Div,
+		DivEqual,
 		Mod,
+		ModEqual,
 		LeftParenthesis,
 		RightParenthesis,
 
@@ -29,22 +38,35 @@ namespace CrossCompiler
 		OrOr,
 
 		// Bit
+		LowerLower,
+		LowerLowerEqual,
+		GreaterGreater,
+		GreaterGreaterEqual,
 		And,
+		AndEqual,
 		Or,
+		OrEqual,
+		Xor,
+		XorEqual,
 		Not,
 		Neg,
-		Xor,
 
 		// Statements
 		Equal,
-		LeftBracket,
-		RightBracket,
+		LeftBrace,
+		RightBrace,
 		Semicolon,
 		If,
+		Else,
 		For,
 		While,
 		Do,
 		Return,
+		Switch,
+		Case,
+		Break,
+		Default,
+		Goto,
 
 		// Unary
 		PlusPlus,
@@ -53,6 +75,7 @@ namespace CrossCompiler
 		// Types
 		Void,
 		Const,
+
 		Bool,
 		Bool1,
 		Bool2,
@@ -74,6 +97,7 @@ namespace CrossCompiler
 		Bool2x4,
 		Bool3x4,
 		Bool4x4,
+
 		Int,
 		Int1,
 		Int2,
@@ -95,6 +119,7 @@ namespace CrossCompiler
 		Int2x4,
 		Int3x4,
 		Int4x4,
+
 		Uint,
 		Uint1,
 		Uint2,
@@ -116,6 +141,7 @@ namespace CrossCompiler
 		Uint2x4,
 		Uint3x4,
 		Uint4x4,
+
 		Half,
 		Half1,
 		Half2,
@@ -137,6 +163,7 @@ namespace CrossCompiler
 		Half2x4,
 		Half3x4,
 		Half4x4,
+
 		Float,
 		Float1,
 		Float2,
@@ -158,6 +185,7 @@ namespace CrossCompiler
 		Float2x4,
 		Float3x4,
 		Float4x4,
+
 		Texture,
 		Texture1D,
 		Texture1DArray,
@@ -166,6 +194,7 @@ namespace CrossCompiler
 		Texture3D,
 		TextureCube,
 		TextureCubeArray,
+
 		Sampler,
 		Sampler1D,
 		Sampler2D,
@@ -173,12 +202,11 @@ namespace CrossCompiler
 		SamplerCube,
 		SamplerState,
 		SampleComparisonState,
+
 		Buffer,
 		AppendStructuredBuffer,
 		ByteAddressBuffer,
 		ConsumeStructuredBuffer,
-		InputPatch,
-		OutputPatch,
 		RWBuffer,
 		RWByteAddressBuffer,
 		RWStructuredBuffer,
@@ -188,6 +216,8 @@ namespace CrossCompiler
 		RWTexture2DArray,
 		RWTexture3D,
 		StructuredBuffer,
+		InputPatch,
+		OutputPatch,
 
 		// Modifiers
 		In,
@@ -196,10 +226,6 @@ namespace CrossCompiler
 		Static,
 
 		// Misc
-		Identifier,
-		UnsignedInteger,
-		FloatNumber,
-
 		LeftSquareBracket,
 		RightSquareBracket,
 		Question,
@@ -209,19 +235,67 @@ namespace CrossCompiler
 		Struct,
 		CBuffer,
 		GroupShared,
+
+		Identifier,
+		UnsignedIntegerConstant,
+		FloatConstant,
+		BoolConstant,
+	};
+
+	struct FHlslToken
+	{
+		EHlslToken Token;
+		FString String;
+		uint32 UnsignedInteger;
+		float Float;
+
+		FString* SourceFilename;
+		int32 SourceLine;
+		int32 SourceColumn;
+
+		explicit FHlslToken(const FString& Identifier) : Token(EHlslToken::Identifier), String(Identifier), UnsignedInteger(0), Float(0), SourceFilename(nullptr), SourceLine(0), SourceColumn(0) { }
+		explicit FHlslToken(EHlslToken InToken, const FString& Identifier) : Token(InToken), String(Identifier), UnsignedInteger(0), Float(0), SourceFilename(nullptr), SourceLine(0), SourceColumn(0) { }
+		explicit FHlslToken(uint32 InUnsignedInteger) : Token(EHlslToken::UnsignedIntegerConstant), UnsignedInteger(InUnsignedInteger), Float(0), SourceFilename(nullptr), SourceLine(0), SourceColumn(0) { }
+		explicit FHlslToken(float InFloat) : Token(EHlslToken::FloatConstant), UnsignedInteger(0), Float(InFloat), SourceFilename(nullptr), SourceLine(0), SourceColumn(0) { }
+		explicit FHlslToken(bool bInValue) : Token(EHlslToken::BoolConstant), UnsignedInteger(bInValue ? 1 : 0), Float(0), SourceFilename(nullptr), SourceLine(0), SourceColumn(0) { }
 	};
 
-	class SHADERCOMPILERCOMMON_API FHlslScanner
+	class FHlslScanner
 	{
 	public:
 		FHlslScanner();
 		virtual ~FHlslScanner();
 
-		void Lex(const FString& String);
+		// Processing
+		bool Lex(const FString& String);
 		void Dump();
 
+		// Iterating after Processing
+		uint32 GetCurrentTokenIndex() const
+		{
+			return CurrentToken;
+		}
+
+		void SetCurrentTokenIndex(uint32 NewToken);
+
+		bool HasMoreTokens() const;
+		bool MatchToken(EHlslToken InToken);
+		const FHlslToken* PeekToken(uint32 LookAhead = 0) const;
+		const FHlslToken* GetCurrentToken() const;
+		const FHlslToken* GetCurrentTokenAndAdvance();
+		void Advance() {++CurrentToken;}
+
+		void SourceError(const FString& Error);
+
 	private:
-		struct FHlslScannerData; 
-		FHlslScannerData* Data;
+		TArray<FHlslToken> Tokens;
+		void AddToken(const FHlslToken& Token, const struct FTokenizer& Tokenizer);
+
+		uint32 CurrentToken;
+
+		// Tokens point their source filenames here
+		TIndirectArray<FString> SourceFilenames;
+
+		void Clear();
 	};
 }
diff --git a/Engine/Source/Developer/ShaderCompilerCommon/Private/HlslParser.cpp b/Engine/Source/Developer/ShaderCompilerCommon/Private/HlslParser.cpp
new file mode 100644
--- /dev/null
+++ b/Engine/Source/Developer/ShaderCompilerCommon/Private/HlslParser.cpp
@@ -0,0 +1,997 @@
+// Copyright 1998-2014 Epic Games, Inc. All Rights Reserved.
+
+/*=============================================================================
+	HlslParser.cpp - Implementation for parsing hlsl.
+=============================================================================*/
+
+#include "ShaderCompilerCommon.h"
+#include "HlslParser.h"
+#include "HlslExpressionParser.inl"
+
+namespace CrossCompiler
+{
+	typedef EParseResult(*TTryRule)(class FHlslParser& Scanner);
+	struct FRulePair
+	{
+		EHlslToken Token;
+		TTryRule OptionalPrefix;
+		TTryRule TryRule;
+
+		FRulePair(EHlslToken InToken, TTryRule InTryRule, TTryRule InOptionalPrefix = nullptr) : Token(InToken), TryRule(InTryRule), OptionalPrefix(InOptionalPrefix) {}
+	};
+	typedef TArray<FRulePair> TRulesArray;
+
+	class FHlslParser
+	{
+	public:
+		FHlslParser();
+		FHlslScanner Scanner;
+		FSymbolScope GlobalScope;
+		FSymbolScope* CurrentScope;
+	};
+
+	TRulesArray RulesTranslationUnit;
+	TRulesArray RulesStatements;
+
+	EParseResult TryRules(const TRulesArray& Rules, FHlslParser& Parser, const FString& RuleNames, bool bErrorIfNoMatch)
+	{
+		for (const auto& Rule : Rules)
+		{
+			auto CurrentTokenIndex = Parser.Scanner.GetCurrentTokenIndex();
+			if (Rule.OptionalPrefix)
+			{
+				if ((*Rule.OptionalPrefix)(Parser) == EParseResult::Error)
+				{
+					return EParseResult::Error;
+				}
+			}
+
+			if (Parser.Scanner.MatchToken(Rule.Token) || Rule.Token == EHlslToken::Invalid)
+			{
+				EParseResult Result = (*Rule.TryRule)(Parser);
+				if (Result == EParseResult::Error || Result == EParseResult::Matched)
+				{
+					return Result;
+				}
+			}
+
+			Parser.Scanner.SetCurrentTokenIndex(CurrentTokenIndex);
+		}
+
+		if (bErrorIfNoMatch)
+		{
+			Parser.Scanner.SourceError(FString::Printf(TEXT("No matching %s rules found!"), *RuleNames));
+			return EParseResult::Error;
+		}
+
+		return EParseResult::NotMatched;
+	}
+
+	EParseResult ParseArrayBracketsAndIndex(FHlslParser& Parser)
+	{
+		if (Parser.Scanner.MatchToken(EHlslToken::LeftSquareBracket))
+		{
+			if (ParseExpression(Parser.Scanner, Parser.CurrentScope) != EParseResult::Matched)
+			{
+				Parser.Scanner.SourceError(TEXT("Expected expression!"));
+				return EParseResult::Error;
+			}
+
+			if (!Parser.Scanner.MatchToken(EHlslToken::RightSquareBracket))
+			{
+				Parser.Scanner.SourceError(TEXT("Expected ']'!"));
+				return EParseResult::Error;
+			}
+
+			return EParseResult::Matched;
+		}
+
+		return EParseResult::NotMatched;
+	}
+
+	EParseResult ParseMultiArrayBracketsAndIndex(FHlslParser& Parser)
+	{
+		bool bFoundOne = false;
+		do
+		{
+			auto Result = ParseArrayBracketsAndIndex(Parser);
+			if (Result == EParseResult::Error)
+			{
+				return EParseResult::Error;
+			}
+			else if (Result == EParseResult::NotMatched)
+			{
+				break;
+			}
+
+			bFoundOne = true;
+		}
+		while (Parser.Scanner.HasMoreTokens());
+
+		return bFoundOne ? EParseResult::Matched : EParseResult::NotMatched;
+	}
+
+	EParseResult ParseTextureOrBufferSimpleDeclaration(FHlslParser& Parser, bool bMultiple)
+	{
+		auto OriginalToken = Parser.Scanner.GetCurrentTokenIndex();
+		if (ParseTextureOrBufferType(Parser.Scanner) == EParseResult::Matched)
+		{
+NextParameter:
+			// Handle 'Sampler2D Sampler'
+			if (ParseTextureOrBufferType(Parser.Scanner) == EParseResult::Matched)
+			{
+				//...
+			}
+			else if (!Parser.Scanner.MatchToken(EHlslToken::Identifier))
+			{
+				Parser.Scanner.SourceError(TEXT("Expected Identifier!"));
+				return EParseResult::Error;
+			}
+
+			if (ParseMultiArrayBracketsAndIndex(Parser) == EParseResult::Error)
+			{
+				return EParseResult::Error;
+			}
+
+			if (bMultiple && Parser.Scanner.MatchToken(EHlslToken::Comma))
+			{
+				goto NextParameter;
+			}
+
+			return EParseResult::Matched;
+		}
+
+		Parser.Scanner.SetCurrentTokenIndex(OriginalToken);
+		return EParseResult::NotMatched;
+	}
+
+	// Multi declaration parser flags
+	enum EDeclarationFlags
+	{
+		//EDF_ROLLBACK_IF_NO_MATCH		= 0x0001,
+		EDF_CONST						= 0x0002,
+		EDF_STATIC						= 0x0004,
+		EDF_TEXTURE_SAMPLER_OR_BUFFER	= 0x0008,
+		EDF_INITIALIZER					= 0x0010,
+		EDF_INITIALIZER_LIST			= 0x0020 | EDF_INITIALIZER,
+		EDF_SEMANTIC					= 0x0040,
+		EDF_SEMICOLON					= 0x0080,
+		EDF_IN_OUT						= 0x0100,
+		EDF_MULTIPLE					= 0x0200,
+	};
+
+	EParseResult ParseInitializer(FHlslParser& Parser, bool bAllowLists)
+	{
+		if (bAllowLists && Parser.Scanner.MatchToken(EHlslToken::LeftBrace))
+		{
+			bool bFoundExpr = false;
+			while (Parser.Scanner.HasMoreTokens())
+			{
+				auto Result = ParseExpression(Parser.Scanner, Parser.CurrentScope);
+				if (Result == EParseResult::Error)
+				{
+					Parser.Scanner.SourceError(TEXT("Invalid initielizer expression list\n"));
+					return EParseResult::Error;
+				}
+				else if (Result == EParseResult::NotMatched)
+				{
+					Parser.Scanner.SourceError(TEXT("Expected initializer expression\n"));
+					return EParseResult::Error;
+				}
+				
+				bFoundExpr = true;
+				if (Parser.Scanner.MatchToken(EHlslToken::Comma))
+				{
+					continue;
+				}
+				else if (Parser.Scanner.MatchToken(EHlslToken::RightBrace))
+				{
+					return EParseResult::Matched;
+				}
+			}
+
+			Parser.Scanner.SourceError(TEXT("Expected '}'\n"));
+		}
+		else
+		{
+			auto Result = ParseExpression(Parser.Scanner, Parser.CurrentScope);
+			if (Result == EParseResult::Error)
+			{
+				Parser.Scanner.SourceError(TEXT("Invalid initializer expression\n"));
+			}
+
+			return Result;
+		}
+
+		return EParseResult::NotMatched;
+	}
+
+
+	EParseResult ParseGeneralDeclaration(FHlslParser& Parser, int32 Flags)
+	{
+		auto OriginalToken = Parser.Scanner.GetCurrentTokenIndex();
+
+		int32 StaticFound = 0;
+		int32 ConstFound = 0;
+		int32 InFound = 0;
+		int32 OutFound = 0;
+		int32 InOutFound = 0;
+
+		while (Parser.Scanner.HasMoreTokens())
+		{
+			bool bFound = false;
+			if ((Flags & EDF_STATIC) && Parser.Scanner.MatchToken(EHlslToken::Static))
+			{
+				++StaticFound;
+				if (StaticFound > 1)
+				{
+					Parser.Scanner.SourceError(TEXT("'static' found more than once!\n"));
+				}
+			}
+			else if ((Flags & EDF_CONST) && Parser.Scanner.MatchToken(EHlslToken::Const))
+			{
+				++ConstFound;
+				if (ConstFound > 1)
+				{
+					Parser.Scanner.SourceError(TEXT("'const' found more than once!\n"));
+				}
+			}
+			else if ((Flags & EDF_IN_OUT) && Parser.Scanner.MatchToken(EHlslToken::In))
+			{
+				++InFound;
+				if (InFound > 1)
+				{
+					Parser.Scanner.SourceError(TEXT("'in' found more than once!\n"));
+				}
+			}
+			else if ((Flags & EDF_IN_OUT) && Parser.Scanner.MatchToken(EHlslToken::Out))
+			{
+				++OutFound;
+				if (OutFound > 1)
+				{
+					Parser.Scanner.SourceError(TEXT("'out' found more than once!\n"));
+				}
+			}
+			else if ((Flags & EDF_IN_OUT) && Parser.Scanner.MatchToken(EHlslToken::InOut))
+			{
+				++InOutFound;
+				if (InOutFound > 1)
+				{
+					Parser.Scanner.SourceError(TEXT("'inout' found more than once!\n"));
+				}
+			}
+			else
+			{
+				break;
+			}
+		}
+
+		bool bCanBeUnmatched = (ConstFound + InFound + OutFound + InOutFound + StaticFound) == 0;
+
+		bool bTryBasicType = true;
+		if (Flags & EDF_TEXTURE_SAMPLER_OR_BUFFER)
+		{
+			auto Result = ParseTextureOrBufferSimpleDeclaration(Parser, (Flags & EDF_MULTIPLE) == EDF_MULTIPLE);
+			if (Result == EParseResult::Error)
+			{
+				return EParseResult::Error;
+			}
+			else if (Result == EParseResult::Matched)
+			{
+				bTryBasicType = false;
+			}
+		}
+		
+		if (bTryBasicType)
+		{
+			if (ParseTypeToken(Parser.Scanner, Parser.CurrentScope) == EParseResult::Matched)
+			{
+				bool bMatched = false;
+				do
+				{
+					if (!Parser.Scanner.MatchToken(EHlslToken::Identifier))
+					{
+						Parser.Scanner.SetCurrentTokenIndex(OriginalToken);
+						return EParseResult::NotMatched;
+					}
+
+					if (ParseMultiArrayBracketsAndIndex(Parser) == EParseResult::Error)
+					{
+						return EParseResult::Error;
+					}
+
+					if (Flags & EDF_SEMANTIC)
+					{
+						if (Parser.Scanner.MatchToken(EHlslToken::Colon))
+						{
+							if (!Parser.Scanner.MatchToken(EHlslToken::Identifier))
+							{
+								Parser.Scanner.SourceError(TEXT("Expected identifier for semantic!"));
+								return EParseResult::Error;
+							}
+						}
+					}
+					else if (Flags & EDF_INITIALIZER)
+					{
+						if (Parser.Scanner.MatchToken(EHlslToken::Equal))
+						{
+							if (ParseInitializer(Parser, (Flags & EDF_INITIALIZER_LIST) == EDF_INITIALIZER_LIST) != EParseResult::Matched)
+							{
+								Parser.Scanner.SourceError(TEXT("Invalid initializer\n"));
+								return EParseResult::Error;
+							}
+						}
+					}
+				}
+				while ((Flags & EDF_MULTIPLE) == EDF_MULTIPLE && Parser.Scanner.MatchToken(EHlslToken::Comma));
+			}
+/*
+			{
+				if (bCanBeUnmatched)
+				{
+					Parser.Scanner.SetCurrentTokenIndex(OriginalToken);
+					return EParseResult::NotMatched;
+				}
+			}
+*/
+		}
+
+		if (Flags & EDF_SEMICOLON)
+		{
+			if (!Parser.Scanner.MatchToken(EHlslToken::Semicolon))
+			{
+				Parser.Scanner.SourceError(TEXT("';' expected!\n"));
+				return EParseResult::Error;
+			}
+		}
+
+		return EParseResult::Matched;
+	}
+
+	EParseResult ParseCBuffer(FHlslParser& Parser)
+	{
+		//FInfo Info;
+		Parser.Scanner.MatchToken(EHlslToken::Identifier);
+		bool bFoundRightBrace = false;
+		if (Parser.Scanner.MatchToken(EHlslToken::LeftBrace))
+		{
+			bool bFoundRightBrace = false;
+			while (Parser.Scanner.HasMoreTokens())
+			{
+				auto Result = ParseGeneralDeclaration(Parser, EDF_CONST | EDF_SEMICOLON | EDF_TEXTURE_SAMPLER_OR_BUFFER);
+				if (Result == EParseResult::Error)
+				{
+					return EParseResult::Error;
+				}
+
+				if (Parser.Scanner.MatchToken(EHlslToken::RightBrace))
+				{
+					bFoundRightBrace = true;
+					break;
+				}
+
+				if (Result == EParseResult::NotMatched)
+				{
+					break;
+				}
+			}
+
+			check(bFoundRightBrace);
+			return EParseResult::Matched;
+		}
+
+		Parser.Scanner.SourceError(TEXT("Expected '}'!"));
+		return EParseResult::Error;
+	}
+
+	EParseResult ParseStruct(FHlslParser& Parser)
+	{
+		//FInfo Info;
+		//@todo-rco: Test me!
+		const auto* Name = Parser.Scanner.GetCurrentToken();
+		if (Name && Parser.Scanner.MatchToken(EHlslToken::Identifier))
+		{
+			Parser.CurrentScope->Add(Name->String);
+		}
+
+		if (!Parser.Scanner.MatchToken(EHlslToken::LeftBrace))
+		{
+			Parser.Scanner.SourceError(TEXT("Expected '{'!"));
+			return EParseResult::Error;
+		}
+
+		bool bFoundRightBrace = false;
+		while (Parser.Scanner.HasMoreTokens())
+		{
+			auto Result = ParseGeneralDeclaration(Parser, EDF_CONST | EDF_SEMICOLON | EDF_SEMANTIC | EDF_TEXTURE_SAMPLER_OR_BUFFER);
+			if (Result == EParseResult::Error)
+			{
+				return EParseResult::Error;
+			}
+
+			if (Parser.Scanner.MatchToken(EHlslToken::RightBrace))
+			{
+				bFoundRightBrace = true;
+				break;
+			}
+
+			if (Result == EParseResult::NotMatched)
+			{
+				break;
+			}
+		}
+
+		if (!bFoundRightBrace)
+		{
+			Parser.Scanner.SourceError(TEXT("Expected '}'!"));
+			return EParseResult::Error;
+		}
+
+		if (!Parser.Scanner.MatchToken(EHlslToken::Semicolon))
+		{
+			Parser.Scanner.SourceError(TEXT("Expected ';'!"));
+			return EParseResult::Error;
+		}
+
+		return EParseResult::Matched;
+	}
+
+	EParseResult ParseFunctionParameterDeclaration(FHlslParser& Parser)
+	{
+		bool bStrictCheck = false;
+
+		while (Parser.Scanner.HasMoreTokens())
+		{
+			auto Result = ParseGeneralDeclaration(Parser, EDF_CONST | EDF_IN_OUT | EDF_TEXTURE_SAMPLER_OR_BUFFER | EDF_INITIALIZER);
+			if (Result == EParseResult::Error)
+			{
+				return EParseResult::Error;
+			}
+
+			if (Parser.Scanner.MatchToken(EHlslToken::Comma))
+			{
+				bStrictCheck = true;
+			}
+			else
+			{
+				break;
+			}
+		}
+
+		return EParseResult::Matched;
+	}
+
+	EParseResult ParseFunctionDeclarator(FHlslParser& Parser)
+	{
+		auto OriginalToken = Parser.Scanner.GetCurrentTokenIndex();
+
+		//FInfo Info;
+		auto Result = ParseBasicType(Parser.Scanner/*, Info*/);
+		if (Result == EParseResult::NotMatched)
+		{
+			Parser.Scanner.SetCurrentTokenIndex(OriginalToken);
+			return EParseResult::NotMatched;
+		}
+		check(Result == EParseResult::Matched);
+
+		if (!Parser.Scanner.MatchToken(EHlslToken::Identifier))
+		{
+			// This could be an error... But we should allow testing for a global variable before any rash decisions
+			Parser.Scanner.SetCurrentTokenIndex(OriginalToken);
+			return EParseResult::NotMatched;
+		}
+
+		if (!Parser.Scanner.MatchToken(EHlslToken::LeftParenthesis))
+		{
+			// This could be an error... But we should allow testing for a global variable before any rash decisions
+			Parser.Scanner.SetCurrentTokenIndex(OriginalToken);
+			return EParseResult::NotMatched;
+		}
+
+		if (Parser.Scanner.MatchToken(EHlslToken::Void))
+		{
+			//...
+		}
+		else if (Parser.Scanner.MatchToken(EHlslToken::RightParenthesis))
+		{
+			return EParseResult::Matched;
+		}
+		else
+		{
+			Result = ParseFunctionParameterDeclaration(Parser);
+			if (Result == EParseResult::Error)
+			{
+				return EParseResult::Error;
+			}
+		}
+
+		if (!Parser.Scanner.MatchToken(EHlslToken::RightParenthesis))
+		{
+			Parser.Scanner.SourceError(TEXT("')' expected"));
+			return EParseResult::Error;
+		}
+
+		return EParseResult::Matched;
+	}
+
+	EParseResult ParseStatement(FHlslParser& Parser)
+	{
+		const auto* Token = Parser.Scanner.PeekToken();
+		if (Token && Token->Token == EHlslToken::RightBrace)
+		{
+			return EParseResult::NotMatched;
+		}
+
+		static FString Statement(TEXT("Statement"));
+		return TryRules(RulesStatements, Parser, Statement, false);
+	}
+
+	EParseResult ParseStatementBlock(FHlslParser& Parser)
+	{
+		while (Parser.Scanner.HasMoreTokens())
+		{
+			auto Result = ParseStatement(Parser);
+			if (Result == EParseResult::NotMatched)
+			{
+				if (Parser.Scanner.MatchToken(EHlslToken::RightBrace))
+				{
+					return EParseResult::Matched;
+				}
+				else
+				{
+					Parser.Scanner.SourceError(TEXT("Statement expected!"));
+					break;
+				}
+			}
+			else if (Result == EParseResult::Error)
+			{
+				break;
+			}
+		}
+
+		Parser.Scanner.SourceError(TEXT("'}' expected!"));
+		return EParseResult::Error;
+	}
+
+	EParseResult ParseFunctionDeclaration(FHlslParser& Parser)
+	{
+		EParseResult Result = ParseFunctionDeclarator(Parser);
+		if (Result == EParseResult::NotMatched || Result == EParseResult::Error)
+		{
+			return Result;
+		}
+
+		if (Parser.Scanner.MatchToken(EHlslToken::Semicolon))
+		{
+			// Forward declare
+			return EParseResult::Matched;
+		}
+		else
+		{
+			// Optional semantic
+			if (Parser.Scanner.MatchToken(EHlslToken::Colon))
+			{
+				if (!Parser.Scanner.MatchToken(EHlslToken::Identifier))
+				{
+					Parser.Scanner.SourceError(TEXT("Identifier for semantic expected"));
+					return EParseResult::Error;
+				}
+			}
+
+			if (!Parser.Scanner.MatchToken(EHlslToken::LeftBrace))
+			{
+				Parser.Scanner.SourceError(TEXT("'{' expected"));
+				return EParseResult::Error;
+			}
+
+			return ParseStatementBlock(Parser);
+		}
+
+		return Result;
+	}
+
+	EParseResult ParseLocalDeclaration(FHlslParser& Parser)
+	{
+		return ParseGeneralDeclaration(Parser, EDF_CONST | EDF_INITIALIZER | EDF_INITIALIZER_LIST | EDF_SEMICOLON | EDF_MULTIPLE);
+	}
+
+	EParseResult ParseGlobalVariableDeclaration(FHlslParser& Parser)
+	{
+		return ParseGeneralDeclaration(Parser, EDF_CONST | EDF_STATIC | EDF_TEXTURE_SAMPLER_OR_BUFFER | EDF_INITIALIZER | EDF_INITIALIZER_LIST | EDF_SEMICOLON | EDF_MULTIPLE);
+	}
+
+	EParseResult ParseReturnStatement(FHlslParser& Parser)
+	{
+		if (Parser.Scanner.MatchToken(EHlslToken::Semicolon))
+		{
+			return EParseResult::Matched;
+		}
+
+		if (ParseExpression(Parser.Scanner, Parser.CurrentScope) != EParseResult::Matched)
+		{
+			Parser.Scanner.SourceError(TEXT("Expression expected"));
+			return EParseResult::Error;
+		}
+
+		if (!Parser.Scanner.MatchToken(EHlslToken::Semicolon))
+		{
+			Parser.Scanner.SourceError(TEXT("';' expected"));
+			return EParseResult::Error;
+		}
+
+		return EParseResult::Matched;
+	}
+
+	EParseResult ParseDoStatement(FHlslParser& Parser)
+	{
+		if (ParseStatement(Parser) != EParseResult::Matched)
+		{
+			Parser.Scanner.SourceError(TEXT("Statement expected"));
+			return EParseResult::Error;
+		}
+
+		if (!Parser.Scanner.MatchToken(EHlslToken::LeftParenthesis))
+		{
+			Parser.Scanner.SourceError(TEXT("'(' expected"));
+			return EParseResult::Error;
+		}
+
+		if (ParseExpression(Parser.Scanner, Parser.CurrentScope) != EParseResult::Matched)
+		{
+			Parser.Scanner.SourceError(TEXT("Expression expected"));
+			return EParseResult::Error;
+		}
+
+		if (!Parser.Scanner.MatchToken(EHlslToken::RightParenthesis))
+		{
+			Parser.Scanner.SourceError(TEXT("')' expected"));
+			return EParseResult::Error;
+		}
+
+		if (ParseStatement(Parser) != EParseResult::Matched)
+		{
+			Parser.Scanner.SourceError(TEXT("Statement expected"));
+			return EParseResult::Error;
+		}
+
+		return EParseResult::Matched;
+	}
+
+	EParseResult ParseWhileStatement(FHlslParser& Parser)
+	{
+		if (!Parser.Scanner.MatchToken(EHlslToken::LeftParenthesis))
+		{
+			Parser.Scanner.SourceError(TEXT("'(' expected"));
+			return EParseResult::Error;
+		}
+
+		if (ParseExpression(Parser.Scanner, Parser.CurrentScope) != EParseResult::Matched)
+		{
+			Parser.Scanner.SourceError(TEXT("Expression expected"));
+			return EParseResult::Error;
+		}
+
+		if (!Parser.Scanner.MatchToken(EHlslToken::RightParenthesis))
+		{
+			Parser.Scanner.SourceError(TEXT("')' expected"));
+			return EParseResult::Error;
+		}
+
+		if (ParseStatement(Parser) != EParseResult::Matched)
+		{
+			Parser.Scanner.SourceError(TEXT("Statement expected"));
+			return EParseResult::Error;
+		}
+
+		return EParseResult::Matched;
+	}
+
+	EParseResult ParseForStatement(FHlslParser& Parser)
+	{
+		if (!Parser.Scanner.MatchToken(EHlslToken::LeftParenthesis))
+		{
+			Parser.Scanner.SourceError(TEXT("Expected '('!\n"));
+			return EParseResult::Error;
+		}
+
+		auto Result = ParseGeneralDeclaration(Parser, EDF_CONST | EDF_INITIALIZER);
+		if (Result == EParseResult::Error)
+		{
+			Parser.Scanner.SourceError(TEXT("Expected expression or declaration!\n"));
+			return EParseResult::Error;
+		}
+		else if (Result == EParseResult::NotMatched)
+		{
+			Result = ParseExpression(Parser.Scanner, Parser.CurrentScope);
+			if (Result == EParseResult::Error)
+			{
+				Parser.Scanner.SourceError(TEXT("Expected expression or declaration!\n"));
+				return EParseResult::Error;
+			}
+		}
+
+		if (!Parser.Scanner.MatchToken(EHlslToken::Semicolon))
+		{
+			Parser.Scanner.SourceError(TEXT("Expected ';'!\n"));
+			return EParseResult::Error;
+		}
+
+		Result = ParseExpression(Parser.Scanner, Parser.CurrentScope);
+		if (Result == EParseResult::Error)
+		{
+			Parser.Scanner.SourceError(TEXT("Expected expression or declaration!\n"));
+			return EParseResult::Error;
+		}
+
+		if (!Parser.Scanner.MatchToken(EHlslToken::Semicolon))
+		{
+			Parser.Scanner.SourceError(TEXT("Expected ';'!\n"));
+			return EParseResult::Error;
+		}
+
+		Result = ParseExpression(Parser.Scanner, Parser.CurrentScope);
+		if (Result == EParseResult::Error)
+		{
+			Parser.Scanner.SourceError(TEXT("Expected expression or declaration!\n"));
+			return EParseResult::Error;
+		}
+
+		if (!Parser.Scanner.MatchToken(EHlslToken::RightParenthesis))
+		{
+			Parser.Scanner.SourceError(TEXT("Expected ')'!\n"));
+			return EParseResult::Error;
+		}
+
+		return ParseStatement(Parser);
+	}
+
+	EParseResult ParseIfStatement(FHlslParser& Parser)
+	{
+		if (!Parser.Scanner.MatchToken(EHlslToken::LeftParenthesis))
+		{
+			Parser.Scanner.SourceError(TEXT("'(' expected"));
+			return EParseResult::Error;
+		}
+
+		if (ParseExpression(Parser.Scanner, Parser.CurrentScope) != EParseResult::Matched)
+		{
+			Parser.Scanner.SourceError(TEXT("Expression expected"));
+			return EParseResult::Error;
+		}
+
+		if (!Parser.Scanner.MatchToken(EHlslToken::RightParenthesis))
+		{
+			Parser.Scanner.SourceError(TEXT("')' expected"));
+			return EParseResult::Error;
+		}
+
+		if (ParseStatement(Parser) != EParseResult::Matched)
+		{
+			Parser.Scanner.SourceError(TEXT("Statement expected"));
+			return EParseResult::Error;
+		}
+
+		if (Parser.Scanner.MatchToken(EHlslToken::Else))
+		{
+			if (ParseStatement(Parser) != EParseResult::Matched)
+			{
+				Parser.Scanner.SourceError(TEXT("Statement expected"));
+				return EParseResult::Error;
+			}
+		}
+
+		return EParseResult::Matched;
+	}
+
+	EParseResult TryParseAttribute(FHlslParser& Parser)
+	{
+		if (Parser.Scanner.MatchToken(EHlslToken::LeftSquareBracket))
+		{
+			if (!Parser.Scanner.MatchToken(EHlslToken::Identifier))
+			{
+				Parser.Scanner.SourceError(TEXT("Incorrect attribute\n"));
+				return EParseResult::Error;
+			}
+
+			if (!Parser.Scanner.MatchToken(EHlslToken::RightSquareBracket))
+			{
+				Parser.Scanner.SourceError(TEXT("Incorrect attribute\n"));
+				return EParseResult::Error;
+			}
+
+			return EParseResult::Matched;
+		}
+
+		return EParseResult::NotMatched;
+	}
+
+	EParseResult ParseSwitchStatement(FHlslParser& Parser)
+	{
+		if (!Parser.Scanner.MatchToken(EHlslToken::LeftParenthesis))
+		{
+			Parser.Scanner.SourceError(TEXT("'(' expected"));
+			return EParseResult::Error;
+		}
+
+		if (ParseExpression(Parser.Scanner, Parser.CurrentScope) != EParseResult::Matched)
+		{
+			Parser.Scanner.SourceError(TEXT("Expression expected"));
+			return EParseResult::Error;
+		}
+
+		if (!Parser.Scanner.MatchToken(EHlslToken::RightParenthesis))
+		{
+			Parser.Scanner.SourceError(TEXT("')' expected"));
+			return EParseResult::Error;
+		}
+
+		if (!Parser.Scanner.MatchToken(EHlslToken::LeftBrace))
+		{
+			Parser.Scanner.SourceError(TEXT("'{' expected"));
+			return EParseResult::Error;
+		}
+
+		bool bDefaultFound = false;
+		while (Parser.Scanner.HasMoreTokens())
+		{
+			if (Parser.Scanner.MatchToken(EHlslToken::Default))
+			{
+				if (bDefaultFound)
+				{
+					Parser.Scanner.SourceError(TEXT("'default' found twice on switch() statement!"));
+					return EParseResult::Error;
+				}
+
+				if (!Parser.Scanner.MatchToken(EHlslToken::Colon))
+				{
+					Parser.Scanner.SourceError(TEXT("':' expected"));
+					return EParseResult::Error;
+				}
+
+				bDefaultFound = true;
+			}
+			else if (Parser.Scanner.MatchToken(EHlslToken::Case))
+			{
+				if (ParseExpression(Parser.Scanner, Parser.CurrentScope) != EParseResult::Matched)
+				{
+					Parser.Scanner.SourceError(TEXT("Expression expected on case label!"));
+					return EParseResult::Error;
+				}
+
+				if (!Parser.Scanner.MatchToken(EHlslToken::Colon))
+				{
+					Parser.Scanner.SourceError(TEXT("':' expected"));
+					return EParseResult::Error;
+				}
+			}
+
+			auto* Peek = Parser.Scanner.PeekToken();
+			if (!Peek)
+			{
+				break;
+			}
+			else if (Peek->Token == EHlslToken::RightBrace)
+			{
+				return EParseResult::Matched;
+			}
+			else if (Peek->Token == EHlslToken::Case || Peek->Token == EHlslToken::Default)
+			{
+				continue;
+			}
+			else
+			{
+				auto Result = ParseStatement(Parser);
+				if (Result == EParseResult::Error)
+				{
+					return EParseResult::Error;
+				}
+			}
+/*
+			// Process statements until hitting break, case, default or '}'
+			while (Parser.Scanner.HasMoreTokens())
+			{
+				if (Parser.Scanner.MatchToken(EHlslToken::RightBrace))
+				{
+					return EParseResult::Matched;
+				}
+
+				if (Parser.Scanner.MatchToken(EHlslToken::Break))
+				{
+					break;
+				}
+			}
+			check(0);
+*/
+		}
+
+		return EParseResult::Error;
+	}
+
+	EParseResult ParseExpressionStatement(FHlslParser& Parser)
+	{
+		auto OriginalToken = Parser.Scanner.GetCurrentTokenIndex();
+		if (ParseExpression(Parser.Scanner, Parser.CurrentScope) == EParseResult::Matched)
+		{
+			if (Parser.Scanner.MatchToken(EHlslToken::Semicolon))
+			{
+				return EParseResult::Matched;
+			}
+		}
+
+		Parser.Scanner.SetCurrentTokenIndex(OriginalToken);
+		return EParseResult::NotMatched;
+	}
+
+	EParseResult ParseEmptyStatement(FHlslParser& Parser)
+	{
+		// Nothing to do here...
+		return EParseResult::Matched;
+	}
+
+	void InitRules()
+	{
+		static bool bInit = false;
+		if (!bInit)
+		{
+			bInit = true;
+
+			RulesTranslationUnit.Add(FRulePair(EHlslToken::Struct, ParseStruct));
+			RulesTranslationUnit.Add(FRulePair(EHlslToken::CBuffer, ParseCBuffer));
+			RulesTranslationUnit.Add(FRulePair(EHlslToken::Invalid, ParseFunctionDeclaration));
+			RulesTranslationUnit.Add(FRulePair(EHlslToken::Invalid, ParseGlobalVariableDeclaration));
+
+			RulesStatements.Add(FRulePair(EHlslToken::LeftBrace, ParseStatementBlock));
+			RulesStatements.Add(FRulePair(EHlslToken::Return, ParseReturnStatement));
+			RulesStatements.Add(FRulePair(EHlslToken::Do, ParseDoStatement));
+			RulesStatements.Add(FRulePair(EHlslToken::While, ParseWhileStatement, TryParseAttribute));
+			RulesStatements.Add(FRulePair(EHlslToken::For, ParseForStatement, TryParseAttribute));
+			RulesStatements.Add(FRulePair(EHlslToken::If, ParseIfStatement, TryParseAttribute));
+			RulesStatements.Add(FRulePair(EHlslToken::Switch, ParseSwitchStatement, TryParseAttribute));
+			RulesStatements.Add(FRulePair(EHlslToken::Semicolon, ParseEmptyStatement));
+			RulesStatements.Add(FRulePair(EHlslToken::Break, ParseEmptyStatement));
+			RulesStatements.Add(FRulePair(EHlslToken::Invalid, ParseExpressionStatement));
+			RulesStatements.Add(FRulePair(EHlslToken::Invalid, ParseLocalDeclaration));
+		}
+	}
+
+	FHlslParser::FHlslParser() :
+		GlobalScope(nullptr)
+	{
+		CurrentScope = &GlobalScope;
+		InitRules();
+	}
+
+	namespace Parser
+	{
+		bool Parse(const FString& Input)
+		{
+			FHlslParser Parser;
+			if (!Parser.Scanner.Lex(Input))
+			{
+				return false;
+			}
+
+			bool bSuccess = true;
+			while (Parser.Scanner.HasMoreTokens())
+			{
+				auto LastIndex = Parser.Scanner.GetCurrentTokenIndex();
+
+				static FString GlobalDeclOrDefinition(TEXT("Global declaration or definition"));
+				auto Result = TryRules(RulesTranslationUnit, Parser, GlobalDeclOrDefinition, true);
+				if (Result == EParseResult::Error)
+				{
+					bSuccess = false;
+					break;
+				}
+				else
+				{
+					check(Result == EParseResult::Matched);
+				}
+
+				check(LastIndex != Parser.Scanner.GetCurrentTokenIndex());
+			}
+
+			return bSuccess;
+		}
+	}
+}
diff --git a/Engine/Source/Developer/ShaderCompilerCommon/Private/HlslParser.h b/Engine/Source/Developer/ShaderCompilerCommon/Private/HlslParser.h
new file mode 100644
--- /dev/null
+++ b/Engine/Source/Developer/ShaderCompilerCommon/Private/HlslParser.h
@@ -0,0 +1,24 @@
+// Copyright 1998-2014 Epic Games, Inc. All Rights Reserved.
+
+/*=============================================================================
+	HlslParser.h - Interface for parsing hlsl.
+=============================================================================*/
+
+#pragma once
+
+#include "HlslLexer.h"
+
+namespace CrossCompiler
+{
+	enum class EParseResult
+	{
+		Matched,
+		NotMatched,
+		Error,
+	};
+
+	namespace Parser
+	{
+		bool Parse(const FString& Input);
+	}
+}
diff --git a/Engine/Source/Programs/CrossCompilerTool/Private/CrossCompilerTool.cpp b/Engine/Source/Programs/CrossCompilerTool/Private/CrossCompilerTool.cpp
--- a/Engine/Source/Programs/CrossCompilerTool/Private/CrossCompilerTool.cpp
+++ b/Engine/Source/Programs/CrossCompilerTool/Private/CrossCompilerTool.cpp
@@ -6,6 +6,8 @@
 #include "hlslcc.h"
 #include "MetalBackend.h"
 #include "GlslBackend.h"
+#include "HlslLexer.h"
+#include "HlslParser.h"
 
 #include "RequiredProgramMainCPPInclude.h"
 
@@ -49,13 +51,22 @@ namespace CCT
 		FString HLSLShaderSource;
 		if (!FFileHelper::LoadFileToString(HLSLShaderSource, *RunInfo.InputFile))
 		{
-			UE_LOG(LogCrossCompilerTool, Error, TEXT("Couldn't load Input file!"));
+			UE_LOG(LogCrossCompilerTool, Error, TEXT("Couldn't load Input file '%s'!"), *RunInfo.InputFile);
 			return 1;
 		}
 
+		if (RunInfo.bUseNew)
+		{
+			// Assume it's preprocessed
+			////Parser.Parse(TEXT("+-1 * 2 + (3 + -4) + F(5) + G(6,7) + A++ - --B"));
+			//Parser.Parse(TEXT("void X() { a += (b ? 1 : 0); }"));
+			//Parser.Parse(TEXT("void X() { return 5 * 3 + 2; }"));
+			CrossCompiler::Parser::Parse(HLSLShaderSource);
+			//Scanner.Dump();
+		}
+
 		ANSICHAR* ShaderSource = 0;
 		ANSICHAR* ErrorLog = 0;
-
 		int Result = HlslCrossCompile(
 			TCHAR_TO_ANSI(*RunInfo.InputFile),
 			TCHAR_TO_ANSI(*HLSLShaderSource),
diff --git a/Engine/Source/Programs/CrossCompilerTool/Private/CrossCompilerTool.h b/Engine/Source/Programs/CrossCompilerTool/Private/CrossCompilerTool.h
--- a/Engine/Source/Programs/CrossCompilerTool/Private/CrossCompilerTool.h
+++ b/Engine/Source/Programs/CrossCompilerTool/Private/CrossCompilerTool.h
@@ -24,6 +24,7 @@ namespace CCT
 		FString OutputFile;
 		EBackend BackEnd;
 		bool bRunCPP;
+		bool bUseNew;
 		//bool bCPPOnly;
 
 		FRunInfo();
diff --git a/Engine/Source/Programs/CrossCompilerTool/Private/CrossCompilerUtils.cpp b/Engine/Source/Programs/CrossCompilerTool/Private/CrossCompilerUtils.cpp
--- a/Engine/Source/Programs/CrossCompilerTool/Private/CrossCompilerUtils.cpp
+++ b/Engine/Source/Programs/CrossCompilerTool/Private/CrossCompilerUtils.cpp
@@ -12,7 +12,8 @@ namespace CCT
 		InputFile(""),
 		OutputFile(""),
 		BackEnd(BE_Invalid),
-		bRunCPP(true)
+		bRunCPP(true),
+		bUseNew(false)
 	{
 	}
 
@@ -243,6 +244,10 @@ namespace CCT
 			{
 				bRunCPP = false;
 			}
+			else if (Switch.StartsWith(TEXT("new")))
+			{
+				bUseNew = true;
+			}
 		}
 
 		// Default to PS
