Fixed Clang compile errors in SCW

diff --git a/Engine/Source/Developer/ShaderCompilerCommon/Private/HlslLexer.cpp b/Engine/Source/Developer/ShaderCompilerCommon/Private/HlslLexer.cpp
--- a/Engine/Source/Developer/ShaderCompilerCommon/Private/HlslLexer.cpp
+++ b/Engine/Source/Developer/ShaderCompilerCommon/Private/HlslLexer.cpp
@@ -608,27 +608,29 @@ namespace CrossCompiler
 				goto NotFloat;
 			}
 
-			// Exponent [eE][+-]?[0-9]+
-			bool bExponentFound = false;
-			if (Match('e') || Match('E'))
 			{
-				Char = Peek();
-				if (Char == '+' || Char == '-')
+				// Exponent [eE][+-]?[0-9]+
+				bool bExponentFound = false;
+				if (Match('e') || Match('E'))
 				{
-					++Current;
+					Char = Peek();
+					if (Char == '+' || Char == '-')
+					{
+						++Current;
+					}
+
+					if (MatchAndSkipDigits())
+					{
+						bExponentFound = true;
+					}
 				}
 
-				if (MatchAndSkipDigits())
+				if (!bExponentFound && !bExpOptional)
 				{
-					bExponentFound = true;
+					goto NotFloat;
 				}
 			}
 
-			if (!bExponentFound && !bExpOptional)
-			{
-				goto NotFloat;
-			}
-
 			// [fF]
 			Char = Peek();
 			if (Char == 'F' || Char == 'f')
diff --git a/Engine/Source/Developer/ShaderCompilerCommon/Private/HlslParser.cpp b/Engine/Source/Developer/ShaderCompilerCommon/Private/HlslParser.cpp
--- a/Engine/Source/Developer/ShaderCompilerCommon/Private/HlslParser.cpp
+++ b/Engine/Source/Developer/ShaderCompilerCommon/Private/HlslParser.cpp
@@ -17,7 +17,7 @@ namespace CrossCompiler
 		TTryRule OptionalPrefix;
 		TTryRule TryRule;
 
-		FRulePair(EHlslToken InToken, TTryRule InTryRule, TTryRule InOptionalPrefix = nullptr) : Token(InToken), TryRule(InTryRule), OptionalPrefix(InOptionalPrefix) {}
+		FRulePair(EHlslToken InToken, TTryRule InTryRule, TTryRule InOptionalPrefix = nullptr) : Token(InToken), OptionalPrefix(InOptionalPrefix), TryRule(InTryRule) {}
 	};
 	typedef TArray<FRulePair> TRulesArray;
 
