UE4 - HlslParser

diff --git a/Engine/Source/Developer/ShaderCompilerCommon/Private/HlslExpressionParser.inl b/Engine/Source/Developer/ShaderCompilerCommon/Private/HlslExpressionParser.inl
--- a/Engine/Source/Developer/ShaderCompilerCommon/Private/HlslExpressionParser.inl
+++ b/Engine/Source/Developer/ShaderCompilerCommon/Private/HlslExpressionParser.inl
@@ -81,11 +81,24 @@ namespace CrossCompiler
 		FInfo& Info;
 	};
 
-	EParseResult ParseBasicType(EHlslToken Token)
+	enum ETypeFlags
+	{
+		ETF_VOID					= 0x0001,
+		ETF_BUILTIN_NUMERIC			= 0x0002,
+		ETF_SAMPLER_TEXTURE_BUFFER	= 0x0004,
+		ETF_USER_TYPES				= 0x0008,
+	};
+
+	EParseResult ParseGeneralType(EHlslToken Token, int32 TypeFlags)
 	{
 		switch (Token)
 		{
 		case EHlslToken::Void:
+			if (TypeFlags & ETF_VOID)
+			{
+				return EParseResult::Matched;
+			}
+			break;
 
 		case EHlslToken::Bool:
 		case EHlslToken::Bool1:
@@ -196,113 +209,82 @@ namespace CrossCompiler
 		case EHlslToken::Float4x2:
 		case EHlslToken::Float4x3:
 		case EHlslToken::Float4x4:
-			return EParseResult::Matched;
-
-		default:
+			if (TypeFlags & ETF_BUILTIN_NUMERIC)
+			{
+				return EParseResult::Matched;
+			}
 			break;
-		}
-
-		return EParseResult::NotMatched;
-	}
 
-	EParseResult ParseBasicType(FHlslScanner& Scanner)
-	{
-		auto* Token = Scanner.PeekToken();
-		if (Token)
-		{
-			if (ParseBasicType(Token->Token) == EParseResult::Matched)
+		case EHlslToken::Texture:
+		case EHlslToken::Texture1D:
+		case EHlslToken::Texture1DArray:
+		case EHlslToken::Texture2D:
+		case EHlslToken::Texture2DArray:
+		case EHlslToken::Texture2DMS:
+		case EHlslToken::Texture2DMSArray:
+		case EHlslToken::Texture3D:
+		case EHlslToken::TextureCube:
+		case EHlslToken::TextureCubeArray:
+
+		case EHlslToken::Buffer:
+		case EHlslToken::AppendStructuredBuffer:
+		case EHlslToken::ByteAddressBuffer:
+		case EHlslToken::ConsumeStructuredBuffer:
+		case EHlslToken::RWBuffer:
+		case EHlslToken::RWByteAddressBuffer:
+		case EHlslToken::RWStructuredBuffer:
+		case EHlslToken::RWTexture1D:
+		case EHlslToken::RWTexture1DArray:
+		case EHlslToken::RWTexture2D:
+		case EHlslToken::RWTexture2DArray:
+		case EHlslToken::RWTexture3D:
+		case EHlslToken::StructuredBuffer:
+
+		case EHlslToken::Sampler:
+		case EHlslToken::Sampler1D:
+		case EHlslToken::Sampler2D:
+		case EHlslToken::Sampler3D:
+		case EHlslToken::SamplerCube:
+		case EHlslToken::SamplerState:
+		case EHlslToken::SamplerComparisonState:
+			if (TypeFlags & ETF_SAMPLER_TEXTURE_BUFFER)
 			{
-				Scanner.Advance();
 				return EParseResult::Matched;
 			}
-
-			return EParseResult::NotMatched;
+			break;
 		}
-
-		return EParseResult::Error;
+		
+		return EParseResult::NotMatched;
 	}
 
-	EParseResult ParseTextureOrBufferType(FHlslScanner& Scanner)
+	EParseResult ParseGeneralType(const FHlslToken* Token, int32 TypeFlags, FSymbolScope* SymbolScope)
 	{
-		auto* Token = Scanner.PeekToken();
 		if (Token)
 		{
-			switch (Token->Token)
+			if (ParseGeneralType(Token->Token, TypeFlags) == EParseResult::Matched)
 			{
-			case EHlslToken::Texture:
-			case EHlslToken::Texture1D:
-			case EHlslToken::Texture1DArray:
-			case EHlslToken::Texture2D:
-			case EHlslToken::Texture2DArray:
-			case EHlslToken::Texture3D:
-			case EHlslToken::TextureCube:
-			case EHlslToken::TextureCubeArray:
-
-			case EHlslToken::Buffer:
-			case EHlslToken::AppendStructuredBuffer:
-			case EHlslToken::ByteAddressBuffer:
-			case EHlslToken::ConsumeStructuredBuffer:
-			case EHlslToken::RWBuffer:
-			case EHlslToken::RWByteAddressBuffer:
-			case EHlslToken::RWStructuredBuffer:
-			case EHlslToken::RWTexture1D:
-			case EHlslToken::RWTexture1DArray:
-			case EHlslToken::RWTexture2D:
-			case EHlslToken::RWTexture2DArray:
-			case EHlslToken::RWTexture3D:
-			case EHlslToken::StructuredBuffer:
-
-			case EHlslToken::Sampler:
-			case EHlslToken::Sampler1D:
-			case EHlslToken::Sampler2D:
-			case EHlslToken::Sampler3D:
-			case EHlslToken::SamplerCube:
-			case EHlslToken::SamplerState:
-			case EHlslToken::SampleComparisonState:
-				Scanner.Advance();
 				return EParseResult::Matched;
-
-			default:
-				break;
 			}
 
-			return EParseResult::NotMatched;
-		}
-		else
-		{
-			return EParseResult::Error;
-		}
-	}
+			if (TypeFlags & ETF_USER_TYPES)
+			{
+				check(SymbolScope);
+				if (Token->Token == EHlslToken::Identifier && SymbolScope->FindType(Token->String))
+				{
+					return EParseResult::Matched;
+				}
+			}
 
-	EParseResult ParseTypeToken(const FHlslToken* Token, FSymbolScope* SymbolScope)
-	{
-		if (Token == nullptr)
-		{
 			return EParseResult::NotMatched;
 		}
 
-		if (ParseBasicType(Token->Token) == EParseResult::Matched)
-		{
-			return EParseResult::Matched;
-		}
-		else if (Token->Token == EHlslToken::Identifier && SymbolScope->FindType(Token->String))
-		{
-			return EParseResult::Matched;
-		}
-
-		return EParseResult::NotMatched;
+		return EParseResult::Error;
 	}
 
-
-	EParseResult ParseTypeToken(FHlslScanner& Scanner, FSymbolScope* SymbolScope)
+	EParseResult ParseGeneralType(FHlslScanner& Scanner, int32 TypeFlags, FSymbolScope* SymbolScope = nullptr)
 	{
-		if (ParseBasicType(Scanner/*, Info*/) == EParseResult::Matched)
-		{
-			return EParseResult::Matched;
-		}
-
-		const auto* Token = Scanner.GetCurrentToken();
-		if (Token && Token->Token == EHlslToken::Identifier && SymbolScope->FindType(Token->String))
+		auto* Token = Scanner.PeekToken();
+		if (ParseGeneralType(Token, TypeFlags, SymbolScope) == EParseResult::Matched)
 		{
 			Scanner.Advance();
 			return EParseResult::Matched;
@@ -360,6 +342,7 @@ namespace CrossCompiler
 				Scanner.Advance();
 				break;
 
+			case EHlslToken::BoolConstant:
 			case EHlslToken::UnsignedIntegerConstant:
 				//FPlatformMisc::LowLevelOutputDebugStringf(TEXT("Atom Int %d\n"), Token->UnsignedInteger);
 				Scanner.Advance();
@@ -386,7 +369,7 @@ namespace CrossCompiler
 
 				const auto* Peek1 = Scanner.PeekToken(0);
 				const auto* Peek2 = Scanner.PeekToken(1);
-				if (Peek1 && ParseTypeToken(Peek1, SymbolScope) == EParseResult::Matched && Peek2 && Peek2->Token == EHlslToken::RightParenthesis)
+				if (Peek1 && ParseGeneralType(Peek1, ETF_BUILTIN_NUMERIC | ETF_USER_TYPES, SymbolScope) == EParseResult::Matched && Peek2 && Peek2->Token == EHlslToken::RightParenthesis)
 				{
 					// Cast
 					Info.PrintWithTabs(FString::Printf(TEXT("Cast %s\n"), *Peek1->String));
@@ -438,14 +421,14 @@ namespace CrossCompiler
 				break;
 			default:
 				// Grrr handle Sampler as a variable name
-				if (ParseTextureOrBufferType(Scanner) == EParseResult::Matched)
+				if (ParseGeneralType(Scanner, ETF_SAMPLER_TEXTURE_BUFFER) == EParseResult::Matched)
 				{
 					bFoundUnary = false;
 					bFoundNonUnary = true;
 					break;
 				}
 				// Handle float3(x,y,z)
-				else if (ParseBasicType(Scanner) == EParseResult::Matched)
+				else if (ParseGeneralType(Scanner, ETF_BUILTIN_NUMERIC) == EParseResult::Matched)
 				{
 					bFoundUnary = false;
 					//bFoundTypeInitializerList = true;
diff --git a/Engine/Source/Developer/ShaderCompilerCommon/Private/HlslLexer.cpp b/Engine/Source/Developer/ShaderCompilerCommon/Private/HlslLexer.cpp
--- a/Engine/Source/Developer/ShaderCompilerCommon/Private/HlslLexer.cpp
+++ b/Engine/Source/Developer/ShaderCompilerCommon/Private/HlslLexer.cpp
@@ -334,6 +334,8 @@ namespace CrossCompiler
 		InsertToken(TEXT("Texture1DArray"), EHlslToken::Texture1DArray);
 		InsertToken(TEXT("Texture2D"), EHlslToken::Texture2D);
 		InsertToken(TEXT("Texture2DArray"), EHlslToken::Texture2DArray);
+		InsertToken(TEXT("Texture2DMS"), EHlslToken::Texture2DMS);
+		InsertToken(TEXT("Texture2DMSArray"), EHlslToken::Texture2DMSArray);
 		InsertToken(TEXT("Texture3D"), EHlslToken::Texture3D);
 		InsertToken(TEXT("TextureCube"), EHlslToken::TextureCube);
 		InsertToken(TEXT("TextureCubeArray"), EHlslToken::TextureCubeArray);
@@ -344,7 +346,7 @@ namespace CrossCompiler
 		InsertToken(TEXT("Sampler3D"), EHlslToken::Sampler3D);
 		InsertToken(TEXT("SamplerCube"), EHlslToken::SamplerCube);
 		InsertToken(TEXT("SamplerState"), EHlslToken::SamplerState);
-		InsertToken(TEXT("SampleComparisonState"), EHlslToken::SampleComparisonState);
+		InsertToken(TEXT("SamplerComparisonState"), EHlslToken::SamplerComparisonState);
 
 		InsertToken(TEXT("Buffer"), EHlslToken::Buffer);
 		InsertToken(TEXT("AppendStructuredBuffer"), EHlslToken::AppendStructuredBuffer);
@@ -378,6 +380,8 @@ namespace CrossCompiler
 		InsertToken(TEXT("struct"), EHlslToken::Struct);
 		InsertToken(TEXT("cbuffer"), EHlslToken::CBuffer);
 		InsertToken(TEXT("groupshared"), EHlslToken::GroupShared);
+		InsertToken(TEXT("nointerpolation"), EHlslToken::NoInterpolation);
+		InsertToken(TEXT("row_major"), EHlslToken::RowMajor);
 	}
 
 	struct FTokenizer
@@ -515,11 +519,16 @@ namespace CrossCompiler
 			{
 				auto Char = Peek();
 				++Current;
-				if (Char == '\n')
+				if (Char == '\r' && Peek() == '\n')
 				{
+					++Current;
+					break;
+				}
+				else if (Char == '\n')
+				{
+
 					break;
 				}
-				check(Char != '\r');
 			}
 
 			++Line;
@@ -882,16 +891,16 @@ namespace CrossCompiler
 		Tokens[TokenIndex].SourceColumn = (int32)(Tokenizer.Current - Tokenizer.CurrentLineStart) + 1;
 	}
 
-	void FHlslScanner::Clear()
+	void FHlslScanner::Clear(const FString& Filename)
 	{
 		Tokens.Empty();
-		new (SourceFilenames) FString(TEXT(""));
+		new (SourceFilenames) FString(Filename);
 	}
 
-	bool FHlslScanner::Lex(const FString& String)
+	bool FHlslScanner::Lex(const FString& String, const FString& Filename)
 	{
-		Clear();
-
+		Clear(Filename);
+		
 		// Simple heuristic to avoid reallocating
 		Tokens.Reserve(String.Len() / 4);
 
diff --git a/Engine/Source/Developer/ShaderCompilerCommon/Private/HlslLexer.h b/Engine/Source/Developer/ShaderCompilerCommon/Private/HlslLexer.h
--- a/Engine/Source/Developer/ShaderCompilerCommon/Private/HlslLexer.h
+++ b/Engine/Source/Developer/ShaderCompilerCommon/Private/HlslLexer.h
@@ -191,6 +191,8 @@ namespace CrossCompiler
 		Texture1DArray,
 		Texture2D,
 		Texture2DArray,
+		Texture2DMS,
+		Texture2DMSArray,
 		Texture3D,
 		TextureCube,
 		TextureCubeArray,
@@ -201,7 +203,7 @@ namespace CrossCompiler
 		Sampler3D,
 		SamplerCube,
 		SamplerState,
-		SampleComparisonState,
+		SamplerComparisonState,
 
 		Buffer,
 		AppendStructuredBuffer,
@@ -235,6 +237,8 @@ namespace CrossCompiler
 		Struct,
 		CBuffer,
 		GroupShared,
+		NoInterpolation,
+		RowMajor,
 
 		Identifier,
 		UnsignedIntegerConstant,
@@ -267,7 +271,7 @@ namespace CrossCompiler
 		virtual ~FHlslScanner();
 
 		// Processing
-		bool Lex(const FString& String);
+		bool Lex(const FString& String, const FString& Filename);
 		void Dump();
 
 		// Iterating after Processing
@@ -296,6 +300,6 @@ namespace CrossCompiler
 		// Tokens point their source filenames here
 		TIndirectArray<FString> SourceFilenames;
 
-		void Clear();
+		void Clear(const FString& Filename);
 	};
 }
diff --git a/Engine/Source/Developer/ShaderCompilerCommon/Private/HlslParser.cpp b/Engine/Source/Developer/ShaderCompilerCommon/Private/HlslParser.cpp
--- a/Engine/Source/Developer/ShaderCompilerCommon/Private/HlslParser.cpp
+++ b/Engine/Source/Developer/ShaderCompilerCommon/Private/HlslParser.cpp
@@ -67,14 +67,25 @@ namespace CrossCompiler
 		return EParseResult::NotMatched;
 	}
 
-	EParseResult ParseArrayBracketsAndIndex(FHlslParser& Parser)
+	EParseResult ParseArrayBracketsAndIndex(FHlslParser& Parser, bool bNeedsDimension)
 	{
 		if (Parser.Scanner.MatchToken(EHlslToken::LeftSquareBracket))
 		{
-			if (ParseExpression(Parser.Scanner, Parser.CurrentScope) != EParseResult::Matched)
+			if (bNeedsDimension)
 			{
-				Parser.Scanner.SourceError(TEXT("Expected expression!"));
-				return EParseResult::Error;
+				if (ParseExpression(Parser.Scanner, Parser.CurrentScope) != EParseResult::Matched)
+				{
+					Parser.Scanner.SourceError(TEXT("Expected expression!"));
+					return EParseResult::Error;
+				}
+			}
+			else
+			{
+				if (ParseExpression(Parser.Scanner, Parser.CurrentScope) == EParseResult::Error)
+				{
+					Parser.Scanner.SourceError(TEXT("Expected expression!"));
+					return EParseResult::Error;
+				}
 			}
 
 			if (!Parser.Scanner.MatchToken(EHlslToken::RightSquareBracket))
@@ -89,12 +100,12 @@ namespace CrossCompiler
 		return EParseResult::NotMatched;
 	}
 
-	EParseResult ParseMultiArrayBracketsAndIndex(FHlslParser& Parser)
+	EParseResult ParseMultiArrayBracketsAndIndex(FHlslParser& Parser, bool bNeedsDimension)
 	{
 		bool bFoundOne = false;
 		do
 		{
-			auto Result = ParseArrayBracketsAndIndex(Parser);
+			auto Result = ParseArrayBracketsAndIndex(Parser, bNeedsDimension);
 			if (Result == EParseResult::Error)
 			{
 				return EParseResult::Error;
@@ -114,11 +125,40 @@ namespace CrossCompiler
 	EParseResult ParseTextureOrBufferSimpleDeclaration(FHlslParser& Parser, bool bMultiple)
 	{
 		auto OriginalToken = Parser.Scanner.GetCurrentTokenIndex();
-		if (ParseTextureOrBufferType(Parser.Scanner) == EParseResult::Matched)
+		if (ParseGeneralType(Parser.Scanner, ETF_SAMPLER_TEXTURE_BUFFER) != EParseResult::Matched)
+		{
+			goto Unmatched;
+		}
+
+		if (Parser.Scanner.MatchToken(EHlslToken::Lower))
+		{
+			auto Result = ParseGeneralType(Parser.Scanner, ETF_BUILTIN_NUMERIC | ETF_USER_TYPES, Parser.CurrentScope);
+			if (Result != EParseResult::Matched)
+			{
+				Parser.Scanner.SourceError(TEXT("Expected type!"));
+				return EParseResult::Error;
+			}
+
+			if (Parser.Scanner.MatchToken(EHlslToken::Comma))
+			{
+				if (!Parser.Scanner.MatchToken(EHlslToken::UnsignedIntegerConstant))
+				{
+					Parser.Scanner.SourceError(TEXT("Expected constant!"));
+					return EParseResult::Error;
+				}
+			}
+
+			if (!Parser.Scanner.MatchToken(EHlslToken::Greater))
+			{
+				Parser.Scanner.SourceError(TEXT("Expected '>'!"));
+				return EParseResult::Error;
+			}
+		}
+
 		{
 NextParameter:
 			// Handle 'Sampler2D Sampler'
-			if (ParseTextureOrBufferType(Parser.Scanner) == EParseResult::Matched)
+			if (ParseGeneralType(Parser.Scanner, ETF_SAMPLER_TEXTURE_BUFFER) == EParseResult::Matched)
 			{
 				//...
 			}
@@ -128,7 +168,7 @@ NextParameter:
 				return EParseResult::Error;
 			}
 
-			if (ParseMultiArrayBracketsAndIndex(Parser) == EParseResult::Error)
+			if (ParseMultiArrayBracketsAndIndex(Parser, true) == EParseResult::Error)
 			{
 				return EParseResult::Error;
 			}
@@ -141,6 +181,7 @@ NextParameter:
 			return EParseResult::Matched;
 		}
 
+Unmatched:
 		Parser.Scanner.SetCurrentTokenIndex(OriginalToken);
 		return EParseResult::NotMatched;
 	}
@@ -149,7 +190,7 @@ NextParameter:
 	enum EDeclarationFlags
 	{
 		//EDF_ROLLBACK_IF_NO_MATCH		= 0x0001,
-		EDF_CONST						= 0x0002,
+		EDF_CONST_ROW_MAJOR				= 0x0002,
 		EDF_STATIC						= 0x0004,
 		EDF_TEXTURE_SAMPLER_OR_BUFFER	= 0x0008,
 		EDF_INITIALIZER					= 0x0010,
@@ -158,6 +199,9 @@ NextParameter:
 		EDF_SEMICOLON					= 0x0080,
 		EDF_IN_OUT						= 0x0100,
 		EDF_MULTIPLE					= 0x0200,
+		EDF_PRIMITIVE_DATA_TYPE			= 0x0400,
+		EDF_SHARED						= 0x0800,
+		EDF_NOINTERPOLATION				= 0x1000,
 	};
 
 	EParseResult ParseInitializer(FHlslParser& Parser, bool bAllowLists)
@@ -167,6 +211,10 @@ NextParameter:
 			bool bFoundExpr = false;
 			while (Parser.Scanner.HasMoreTokens())
 			{
+				if (Parser.Scanner.MatchToken(EHlslToken::RightBrace))
+				{
+					return EParseResult::Matched;
+				}
 				auto Result = ParseExpression(Parser.Scanner, Parser.CurrentScope);
 				if (Result == EParseResult::Error)
 				{
@@ -188,6 +236,9 @@ NextParameter:
 				{
 					return EParseResult::Matched;
 				}
+				
+				Parser.Scanner.SourceError(TEXT("Expected ','\n"));
+				break;
 			}
 
 			Parser.Scanner.SourceError(TEXT("Expected '}'\n"));
@@ -212,10 +263,30 @@ NextParameter:
 		auto OriginalToken = Parser.Scanner.GetCurrentTokenIndex();
 
 		int32 StaticFound = 0;
+		int32 SharedFound = 0;
 		int32 ConstFound = 0;
+		int32 RowMajorFound = 0;
 		int32 InFound = 0;
 		int32 OutFound = 0;
 		int32 InOutFound = 0;
+		int32 PrimitiveFound = 0;
+
+		if (Flags & EDF_PRIMITIVE_DATA_TYPE)
+		{
+			const auto* Token = Parser.Scanner.GetCurrentToken();
+			if (Token && Token->Token == EHlslToken::Identifier)
+			{
+				if (Token->String == TEXT("point") ||
+					Token->String == TEXT("line") ||
+					Token->String == TEXT("triangle") ||
+					Token->String == TEXT("lineadj") ||
+					Token->String == TEXT("triangleadj"))
+				{
+					Parser.Scanner.Advance();
+					++PrimitiveFound;
+				}
+			}
+		}
 
 		while (Parser.Scanner.HasMoreTokens())
 		{
@@ -228,7 +299,23 @@ NextParameter:
 					Parser.Scanner.SourceError(TEXT("'static' found more than once!\n"));
 				}
 			}
-			else if ((Flags & EDF_CONST) && Parser.Scanner.MatchToken(EHlslToken::Const))
+			else if ((Flags & EDF_NOINTERPOLATION) && Parser.Scanner.MatchToken(EHlslToken::NoInterpolation))
+			{
+				++SharedFound;
+				if (SharedFound > 1)
+				{
+					Parser.Scanner.SourceError(TEXT("'groupshared' found more than once!\n"));
+				}
+			}
+			else if ((Flags & EDF_SHARED) && Parser.Scanner.MatchToken(EHlslToken::GroupShared))
+			{
+				++SharedFound;
+				if (SharedFound > 1)
+				{
+					Parser.Scanner.SourceError(TEXT("'groupshared' found more than once!\n"));
+				}
+			}
+			else if ((Flags & EDF_CONST_ROW_MAJOR) && Parser.Scanner.MatchToken(EHlslToken::Const))
 			{
 				++ConstFound;
 				if (ConstFound > 1)
@@ -236,6 +323,14 @@ NextParameter:
 					Parser.Scanner.SourceError(TEXT("'const' found more than once!\n"));
 				}
 			}
+			else if ((Flags & EDF_CONST_ROW_MAJOR) && Parser.Scanner.MatchToken(EHlslToken::RowMajor))
+			{
+				++RowMajorFound;
+				if (RowMajorFound > 1)
+				{
+					Parser.Scanner.SourceError(TEXT("'row_major' found more than once!\n"));
+				}
+			}
 			else if ((Flags & EDF_IN_OUT) && Parser.Scanner.MatchToken(EHlslToken::In))
 			{
 				++InFound;
@@ -266,7 +361,48 @@ NextParameter:
 			}
 		}
 
-		bool bCanBeUnmatched = (ConstFound + InFound + OutFound + InOutFound + StaticFound) == 0;
+		bool bCanBeUnmatched = (ConstFound + RowMajorFound + InFound + OutFound + InOutFound + StaticFound + SharedFound + PrimitiveFound) == 0;
+
+		if (PrimitiveFound == 0 && (Flags & EDF_PRIMITIVE_DATA_TYPE))
+		{
+			const auto* Token = Parser.Scanner.GetCurrentToken();
+			if (Token && Token->Token == EHlslToken::Identifier)
+			{
+				if (Token->String == TEXT("PointStream") ||
+					Token->String == TEXT("LineStream") ||
+					Token->String == TEXT("TriangleStream"))
+				{
+					Parser.Scanner.Advance();
+					bCanBeUnmatched = false;
+
+					if (!Parser.Scanner.MatchToken(EHlslToken::Lower))
+					{
+						Parser.Scanner.SourceError(TEXT("Expected '<'!"));
+						return EParseResult::Error;
+					}
+
+					if (ParseGeneralType(Parser.Scanner, ETF_BUILTIN_NUMERIC | ETF_USER_TYPES, Parser.CurrentScope) != EParseResult::Matched)
+					{
+						Parser.Scanner.SourceError(TEXT("Expected type!"));
+						return EParseResult::Error;
+					}
+
+					if (!Parser.Scanner.MatchToken(EHlslToken::Greater))
+					{
+						Parser.Scanner.SourceError(TEXT("Expected '>'!"));
+						return EParseResult::Error;
+					}
+
+					if (!Parser.Scanner.MatchToken(EHlslToken::Identifier))
+					{
+						Parser.Scanner.SourceError(TEXT("Expected identifier!"));
+						return EParseResult::Error;
+					}
+
+					return EParseResult::Matched;
+				}
+			}
+		}
 
 		bool bTryBasicType = true;
 		if (Flags & EDF_TEXTURE_SAMPLER_OR_BUFFER)
@@ -284,7 +420,8 @@ NextParameter:
 		
 		if (bTryBasicType)
 		{
-			if (ParseTypeToken(Parser.Scanner, Parser.CurrentScope) == EParseResult::Matched)
+			auto Result = ParseGeneralType(Parser.Scanner, ETF_BUILTIN_NUMERIC | ETF_USER_TYPES, Parser.CurrentScope);
+			if (Result == EParseResult::Matched)
 			{
 				bool bMatched = false;
 				do
@@ -295,11 +432,12 @@ NextParameter:
 						return EParseResult::NotMatched;
 					}
 
-					if (ParseMultiArrayBracketsAndIndex(Parser) == EParseResult::Error)
+					if (ParseMultiArrayBracketsAndIndex(Parser, false) == EParseResult::Error)
 					{
 						return EParseResult::Error;
 					}
 
+					bool bSemanticFound = false;
 					if (Flags & EDF_SEMANTIC)
 					{
 						if (Parser.Scanner.MatchToken(EHlslToken::Colon))
@@ -309,9 +447,12 @@ NextParameter:
 								Parser.Scanner.SourceError(TEXT("Expected identifier for semantic!"));
 								return EParseResult::Error;
 							}
+
+							bSemanticFound = true;
 						}
 					}
-					else if (Flags & EDF_INITIALIZER)
+					
+					if ((Flags & EDF_INITIALIZER) && !bSemanticFound)
 					{
 						if (Parser.Scanner.MatchToken(EHlslToken::Equal))
 						{
@@ -325,15 +466,11 @@ NextParameter:
 				}
 				while ((Flags & EDF_MULTIPLE) == EDF_MULTIPLE && Parser.Scanner.MatchToken(EHlslToken::Comma));
 			}
-/*
+			else if (bCanBeUnmatched && Result == EParseResult::NotMatched)
 			{
-				if (bCanBeUnmatched)
-				{
-					Parser.Scanner.SetCurrentTokenIndex(OriginalToken);
-					return EParseResult::NotMatched;
-				}
+				Parser.Scanner.SetCurrentTokenIndex(OriginalToken);
+				return EParseResult::NotMatched;
 			}
-*/
 		}
 
 		if (Flags & EDF_SEMICOLON)
@@ -358,19 +495,18 @@ NextParameter:
 			bool bFoundRightBrace = false;
 			while (Parser.Scanner.HasMoreTokens())
 			{
-				auto Result = ParseGeneralDeclaration(Parser, EDF_CONST | EDF_SEMICOLON | EDF_TEXTURE_SAMPLER_OR_BUFFER);
-				if (Result == EParseResult::Error)
-				{
-					return EParseResult::Error;
-				}
-
 				if (Parser.Scanner.MatchToken(EHlslToken::RightBrace))
 				{
 					bFoundRightBrace = true;
 					break;
 				}
 
-				if (Result == EParseResult::NotMatched)
+				auto Result = ParseGeneralDeclaration(Parser, EDF_CONST_ROW_MAJOR | EDF_SEMICOLON | EDF_TEXTURE_SAMPLER_OR_BUFFER);
+				if (Result == EParseResult::Error)
+				{
+					return EParseResult::Error;
+				}
+				else if (Result == EParseResult::NotMatched)
 				{
 					break;
 				}
@@ -394,6 +530,15 @@ NextParameter:
 			Parser.CurrentScope->Add(Name->String);
 		}
 
+		if (Parser.Scanner.MatchToken(EHlslToken::Colon))
+		{
+			if (!Parser.Scanner.MatchToken(EHlslToken::Identifier))
+			{
+				Parser.Scanner.SourceError(TEXT("Identifier expected!\n"));
+				return EParseResult::Error;
+			}
+		}
+
 		if (!Parser.Scanner.MatchToken(EHlslToken::LeftBrace))
 		{
 			Parser.Scanner.SourceError(TEXT("Expected '{'!"));
@@ -403,19 +548,18 @@ NextParameter:
 		bool bFoundRightBrace = false;
 		while (Parser.Scanner.HasMoreTokens())
 		{
-			auto Result = ParseGeneralDeclaration(Parser, EDF_CONST | EDF_SEMICOLON | EDF_SEMANTIC | EDF_TEXTURE_SAMPLER_OR_BUFFER);
-			if (Result == EParseResult::Error)
-			{
-				return EParseResult::Error;
-			}
-
 			if (Parser.Scanner.MatchToken(EHlslToken::RightBrace))
 			{
 				bFoundRightBrace = true;
 				break;
 			}
 
-			if (Result == EParseResult::NotMatched)
+			auto Result = ParseGeneralDeclaration(Parser, EDF_CONST_ROW_MAJOR | EDF_SEMICOLON | EDF_SEMANTIC | EDF_TEXTURE_SAMPLER_OR_BUFFER | EDF_NOINTERPOLATION);
+			if (Result == EParseResult::Error)
+			{
+				return EParseResult::Error;
+			}
+			else if (Result == EParseResult::NotMatched)
 			{
 				break;
 			}
@@ -442,19 +586,19 @@ NextParameter:
 
 		while (Parser.Scanner.HasMoreTokens())
 		{
-			auto Result = ParseGeneralDeclaration(Parser, EDF_CONST | EDF_IN_OUT | EDF_TEXTURE_SAMPLER_OR_BUFFER | EDF_INITIALIZER);
+			auto Result = ParseGeneralDeclaration(Parser, EDF_CONST_ROW_MAJOR | EDF_IN_OUT | EDF_TEXTURE_SAMPLER_OR_BUFFER | EDF_INITIALIZER | EDF_SEMANTIC | EDF_PRIMITIVE_DATA_TYPE | EDF_NOINTERPOLATION);
 			if (Result == EParseResult::Error)
 			{
 				return EParseResult::Error;
 			}
-
-			if (Parser.Scanner.MatchToken(EHlslToken::Comma))
+			else if (!Parser.Scanner.MatchToken(EHlslToken::Comma))
 			{
-				bStrictCheck = true;
+				break;
 			}
-			else
+			else if (Result == EParseResult::NotMatched)
 			{
-				break;
+				Parser.Scanner.SourceError(TEXT("Internal error on function parameter!\n"));
+				return EParseResult::Error;
 			}
 		}
 
@@ -464,9 +608,7 @@ NextParameter:
 	EParseResult ParseFunctionDeclarator(FHlslParser& Parser)
 	{
 		auto OriginalToken = Parser.Scanner.GetCurrentTokenIndex();
-
-		//FInfo Info;
-		auto Result = ParseBasicType(Parser.Scanner/*, Info*/);
+		auto Result = ParseGeneralType(Parser.Scanner, ETF_BUILTIN_NUMERIC | ETF_SAMPLER_TEXTURE_BUFFER | ETF_USER_TYPES | ETF_VOID, Parser.CurrentScope);
 		if (Result == EParseResult::NotMatched)
 		{
 			Parser.Scanner.SetCurrentTokenIndex(OriginalToken);
@@ -592,12 +734,12 @@ NextParameter:
 
 	EParseResult ParseLocalDeclaration(FHlslParser& Parser)
 	{
-		return ParseGeneralDeclaration(Parser, EDF_CONST | EDF_INITIALIZER | EDF_INITIALIZER_LIST | EDF_SEMICOLON | EDF_MULTIPLE);
+		return ParseGeneralDeclaration(Parser, EDF_CONST_ROW_MAJOR | EDF_INITIALIZER | EDF_INITIALIZER_LIST | EDF_SEMICOLON | EDF_MULTIPLE);
 	}
 
 	EParseResult ParseGlobalVariableDeclaration(FHlslParser& Parser)
 	{
-		return ParseGeneralDeclaration(Parser, EDF_CONST | EDF_STATIC | EDF_TEXTURE_SAMPLER_OR_BUFFER | EDF_INITIALIZER | EDF_INITIALIZER_LIST | EDF_SEMICOLON | EDF_MULTIPLE);
+		return ParseGeneralDeclaration(Parser, EDF_CONST_ROW_MAJOR | EDF_STATIC | EDF_SHARED | EDF_TEXTURE_SAMPLER_OR_BUFFER | EDF_INITIALIZER | EDF_INITIALIZER_LIST | EDF_SEMICOLON | EDF_MULTIPLE);
 	}
 
 	EParseResult ParseReturnStatement(FHlslParser& Parser)
@@ -694,7 +836,7 @@ NextParameter:
 			return EParseResult::Error;
 		}
 
-		auto Result = ParseGeneralDeclaration(Parser, EDF_CONST | EDF_INITIALIZER);
+		auto Result = ParseGeneralDeclaration(Parser, EDF_CONST_ROW_MAJOR | EDF_INITIALIZER);
 		if (Result == EParseResult::Error)
 		{
 			Parser.Scanner.SourceError(TEXT("Expected expression or declaration!\n"));
@@ -793,6 +935,30 @@ NextParameter:
 				return EParseResult::Error;
 			}
 
+			if (Parser.Scanner.MatchToken(EHlslToken::LeftParenthesis))
+			{
+				do
+				{
+					if (ParseExpression(Parser.Scanner, Parser.CurrentScope) != EParseResult::Matched)
+					{
+						Parser.Scanner.SourceError(TEXT("Incorrect attribute! Expected expression.\n"));
+						return EParseResult::Error;
+					}
+
+					if (!Parser.Scanner.MatchToken(EHlslToken::Comma))
+					{
+						break;
+					}
+				}
+				while (Parser.Scanner.HasMoreTokens());
+
+				if (!Parser.Scanner.MatchToken(EHlslToken::RightParenthesis))
+				{
+					Parser.Scanner.SourceError(TEXT("Incorrect attribute! Expected ')'.\n"));
+					return EParseResult::Error;
+				}
+			}
+
 			if (!Parser.Scanner.MatchToken(EHlslToken::RightSquareBracket))
 			{
 				Parser.Scanner.SourceError(TEXT("Incorrect attribute\n"));
@@ -870,7 +1036,7 @@ NextParameter:
 			{
 				break;
 			}
-			else if (Peek->Token == EHlslToken::RightBrace)
+			else if (Parser.Scanner.MatchToken(EHlslToken::RightBrace))
 			{
 				return EParseResult::Matched;
 			}
@@ -937,7 +1103,7 @@ NextParameter:
 
 			RulesTranslationUnit.Add(FRulePair(EHlslToken::Struct, ParseStruct));
 			RulesTranslationUnit.Add(FRulePair(EHlslToken::CBuffer, ParseCBuffer));
-			RulesTranslationUnit.Add(FRulePair(EHlslToken::Invalid, ParseFunctionDeclaration));
+			RulesTranslationUnit.Add(FRulePair(EHlslToken::Invalid, ParseFunctionDeclaration, TryParseAttribute));
 			RulesTranslationUnit.Add(FRulePair(EHlslToken::Invalid, ParseGlobalVariableDeclaration));
 
 			RulesStatements.Add(FRulePair(EHlslToken::LeftBrace, ParseStatementBlock));
@@ -963,10 +1129,10 @@ NextParameter:
 
 	namespace Parser
 	{
-		bool Parse(const FString& Input)
+		bool Parse(const FString& Input, const FString& Filename)
 		{
 			FHlslParser Parser;
-			if (!Parser.Scanner.Lex(Input))
+			if (!Parser.Scanner.Lex(Input, Filename))
 			{
 				return false;
 			}
diff --git a/Engine/Source/Developer/ShaderCompilerCommon/Private/HlslParser.h b/Engine/Source/Developer/ShaderCompilerCommon/Private/HlslParser.h
--- a/Engine/Source/Developer/ShaderCompilerCommon/Private/HlslParser.h
+++ b/Engine/Source/Developer/ShaderCompilerCommon/Private/HlslParser.h
@@ -19,6 +19,6 @@ namespace CrossCompiler
 
 	namespace Parser
 	{
-		bool Parse(const FString& Input);
+		bool Parse(const FString& Input, const FString& Filename);
 	}
 }
diff --git a/Engine/Source/Programs/CrossCompilerTool/Private/CrossCompilerTool.cpp b/Engine/Source/Programs/CrossCompilerTool/Private/CrossCompilerTool.cpp
--- a/Engine/Source/Programs/CrossCompilerTool/Private/CrossCompilerTool.cpp
+++ b/Engine/Source/Programs/CrossCompilerTool/Private/CrossCompilerTool.cpp
@@ -57,12 +57,39 @@ namespace CCT
 
 		if (RunInfo.bUseNew)
 		{
+
 			// Assume it's preprocessed
 			////Parser.Parse(TEXT("+-1 * 2 + (3 + -4) + F(5) + G(6,7) + A++ - --B"));
 			//Parser.Parse(TEXT("void X() { a += (b ? 1 : 0); }"));
 			//Parser.Parse(TEXT("void X() { return 5 * 3 + 2; }"));
-			CrossCompiler::Parser::Parse(HLSLShaderSource);
+			if (RunInfo.bList)
+			{
+				TArray<FString> List;
+				FString Temp;
+				while (HLSLShaderSource.Split(TEXT("\r\n"), &Temp, &HLSLShaderSource))
+				{
+					List.Add(Temp);
+				}
+
+				for (auto& File : List)
+				{
+					FString HLSLShader;
+					if (!FFileHelper::LoadFileToString(HLSLShader, *File))
+					{
+						UE_LOG(LogCrossCompilerTool, Error, TEXT("Couldn't load Input file '%s'!"), *RunInfo.InputFile);
+						return 1;
+					}
+					UE_LOG(LogCrossCompilerTool, Log, TEXT("%s!"), *File);
+
+					CrossCompiler::Parser::Parse(HLSLShader, File);
+				}
+			}
+			else
+			{
+				CrossCompiler::Parser::Parse(HLSLShaderSource, *RunInfo.InputFile);
+			}
 			//Scanner.Dump();
+			return 1;
 		}
 
 		ANSICHAR* ShaderSource = 0;
diff --git a/Engine/Source/Programs/CrossCompilerTool/Private/CrossCompilerTool.h b/Engine/Source/Programs/CrossCompilerTool/Private/CrossCompilerTool.h
--- a/Engine/Source/Programs/CrossCompilerTool/Private/CrossCompilerTool.h
+++ b/Engine/Source/Programs/CrossCompilerTool/Private/CrossCompilerTool.h
@@ -25,6 +25,7 @@ namespace CCT
 		EBackend BackEnd;
 		bool bRunCPP;
 		bool bUseNew;
+		bool bList;
 		//bool bCPPOnly;
 
 		FRunInfo();
diff --git a/Engine/Source/Programs/CrossCompilerTool/Private/CrossCompilerUtils.cpp b/Engine/Source/Programs/CrossCompilerTool/Private/CrossCompilerUtils.cpp
--- a/Engine/Source/Programs/CrossCompilerTool/Private/CrossCompilerUtils.cpp
+++ b/Engine/Source/Programs/CrossCompilerTool/Private/CrossCompilerUtils.cpp
@@ -13,7 +13,8 @@ namespace CCT
 		OutputFile(""),
 		BackEnd(BE_Invalid),
 		bRunCPP(true),
-		bUseNew(false)
+		bUseNew(false),
+		bList(false)
 	{
 	}
 
@@ -248,6 +249,10 @@ namespace CCT
 			{
 				bUseNew = true;
 			}
+			else if (Switch.StartsWith(TEXT("list")))
+			{
+				bList = true;
+			}
 		}
 
 		// Default to PS
