UE4 - HlslParser - 99% compliant

diff --git a/Engine/Source/Developer/ShaderCompilerCommon/Private/HlslExpressionParser.inl b/Engine/Source/Developer/ShaderCompilerCommon/Private/HlslExpressionParser.inl
--- a/Engine/Source/Developer/ShaderCompilerCommon/Private/HlslExpressionParser.inl
+++ b/Engine/Source/Developer/ShaderCompilerCommon/Private/HlslExpressionParser.inl
@@ -22,14 +22,38 @@ namespace CrossCompiler
 			Symbols.Add(Type);
 		}
 
-		bool FindType(const FString& Type) const
+		static bool FindType(const FSymbolScope* Scope, const FString& Type)
 		{
-			if (Symbols.Contains(Type))
+			while (Scope)
 			{
-				return true;
+				if (Scope->Symbols.Contains(Type))
+				{
+					return true;
+				}
+
+				Scope = Scope->Parent;
 			}
 
-			return Parent ? Parent->FindType(Type) : false;
+			return false;
+		}
+	};
+
+	struct FCreateSymbolScope
+	{
+		FSymbolScope* Original;
+		FSymbolScope** Current;
+
+		FCreateSymbolScope(FSymbolScope** InCurrent) :
+			Current(InCurrent)
+		{
+			Original = *InCurrent;
+			auto* NewScope = new(Original->Children) FSymbolScope(Original);
+			*Current = NewScope;
+		}
+
+		~FCreateSymbolScope()
+		{
+			*Current = Original;
 		}
 	};
 
@@ -269,7 +293,7 @@ namespace CrossCompiler
 			if (TypeFlags & ETF_USER_TYPES)
 			{
 				check(SymbolScope);
-				if (Token->Token == EHlslToken::Identifier && SymbolScope->FindType(Token->String))
+				if (Token->Token == EHlslToken::Identifier && FSymbolScope::FindType(SymbolScope, Token->String))
 				{
 					return EParseResult::Matched;
 				}
diff --git a/Engine/Source/Developer/ShaderCompilerCommon/Private/HlslParser.cpp b/Engine/Source/Developer/ShaderCompilerCommon/Private/HlslParser.cpp
--- a/Engine/Source/Developer/ShaderCompilerCommon/Private/HlslParser.cpp
+++ b/Engine/Source/Developer/ShaderCompilerCommon/Private/HlslParser.cpp
@@ -67,30 +67,30 @@ namespace CrossCompiler
 		return EParseResult::NotMatched;
 	}
 
-	EParseResult ParseArrayBracketsAndIndex(FHlslParser& Parser, bool bNeedsDimension)
+	EParseResult ParseArrayBracketsAndIndex(FHlslScanner& Scanner, FSymbolScope* SymbolScope, bool bNeedsDimension)
 	{
-		if (Parser.Scanner.MatchToken(EHlslToken::LeftSquareBracket))
+		if (Scanner.MatchToken(EHlslToken::LeftSquareBracket))
 		{
 			if (bNeedsDimension)
 			{
-				if (ParseExpression(Parser.Scanner, Parser.CurrentScope) != EParseResult::Matched)
+				if (ParseExpression(Scanner, SymbolScope) != EParseResult::Matched)
 				{
-					Parser.Scanner.SourceError(TEXT("Expected expression!"));
+					Scanner.SourceError(TEXT("Expected expression!"));
 					return EParseResult::Error;
 				}
 			}
 			else
 			{
-				if (ParseExpression(Parser.Scanner, Parser.CurrentScope) == EParseResult::Error)
+				if (ParseExpression(Scanner, SymbolScope) == EParseResult::Error)
 				{
-					Parser.Scanner.SourceError(TEXT("Expected expression!"));
+					Scanner.SourceError(TEXT("Expected expression!"));
 					return EParseResult::Error;
 				}
 			}
 
-			if (!Parser.Scanner.MatchToken(EHlslToken::RightSquareBracket))
+			if (!Scanner.MatchToken(EHlslToken::RightSquareBracket))
 			{
-				Parser.Scanner.SourceError(TEXT("Expected ']'!"));
+				Scanner.SourceError(TEXT("Expected ']'!"));
 				return EParseResult::Error;
 			}
 
@@ -100,12 +100,12 @@ namespace CrossCompiler
 		return EParseResult::NotMatched;
 	}
 
-	EParseResult ParseMultiArrayBracketsAndIndex(FHlslParser& Parser, bool bNeedsDimension)
+	EParseResult ParseMultiArrayBracketsAndIndex(FHlslScanner& Scanner, FSymbolScope* SymbolScope, bool bNeedsDimension)
 	{
 		bool bFoundOne = false;
 		do
 		{
-			auto Result = ParseArrayBracketsAndIndex(Parser, bNeedsDimension);
+			auto Result = ParseArrayBracketsAndIndex(Scanner, SymbolScope, bNeedsDimension);
 			if (Result == EParseResult::Error)
 			{
 				return EParseResult::Error;
@@ -117,40 +117,40 @@ namespace CrossCompiler
 
 			bFoundOne = true;
 		}
-		while (Parser.Scanner.HasMoreTokens());
+		while (Scanner.HasMoreTokens());
 
 		return bFoundOne ? EParseResult::Matched : EParseResult::NotMatched;
 	}
 
-	EParseResult ParseTextureOrBufferSimpleDeclaration(FHlslParser& Parser, bool bMultiple)
+	EParseResult ParseTextureOrBufferSimpleDeclaration(FHlslScanner& Scanner, FSymbolScope* SymbolScope, bool bMultiple)
 	{
-		auto OriginalToken = Parser.Scanner.GetCurrentTokenIndex();
-		if (ParseGeneralType(Parser.Scanner, ETF_SAMPLER_TEXTURE_BUFFER) != EParseResult::Matched)
+		auto OriginalToken = Scanner.GetCurrentTokenIndex();
+		if (ParseGeneralType(Scanner, ETF_SAMPLER_TEXTURE_BUFFER) != EParseResult::Matched)
 		{
 			goto Unmatched;
 		}
 
-		if (Parser.Scanner.MatchToken(EHlslToken::Lower))
+		if (Scanner.MatchToken(EHlslToken::Lower))
 		{
-			auto Result = ParseGeneralType(Parser.Scanner, ETF_BUILTIN_NUMERIC | ETF_USER_TYPES, Parser.CurrentScope);
+			auto Result = ParseGeneralType(Scanner, ETF_BUILTIN_NUMERIC | ETF_USER_TYPES, SymbolScope);
 			if (Result != EParseResult::Matched)
 			{
-				Parser.Scanner.SourceError(TEXT("Expected type!"));
+				Scanner.SourceError(TEXT("Expected type!"));
 				return EParseResult::Error;
 			}
 
-			if (Parser.Scanner.MatchToken(EHlslToken::Comma))
+			if (Scanner.MatchToken(EHlslToken::Comma))
 			{
-				if (!Parser.Scanner.MatchToken(EHlslToken::UnsignedIntegerConstant))
+				if (!Scanner.MatchToken(EHlslToken::UnsignedIntegerConstant))
 				{
-					Parser.Scanner.SourceError(TEXT("Expected constant!"));
+					Scanner.SourceError(TEXT("Expected constant!"));
 					return EParseResult::Error;
 				}
 			}
 
-			if (!Parser.Scanner.MatchToken(EHlslToken::Greater))
+			if (!Scanner.MatchToken(EHlslToken::Greater))
 			{
-				Parser.Scanner.SourceError(TEXT("Expected '>'!"));
+				Scanner.SourceError(TEXT("Expected '>'!"));
 				return EParseResult::Error;
 			}
 		}
@@ -158,22 +158,22 @@ namespace CrossCompiler
 		{
 NextParameter:
 			// Handle 'Sampler2D Sampler'
-			if (ParseGeneralType(Parser.Scanner, ETF_SAMPLER_TEXTURE_BUFFER) == EParseResult::Matched)
+			if (ParseGeneralType(Scanner, ETF_SAMPLER_TEXTURE_BUFFER) == EParseResult::Matched)
 			{
 				//...
 			}
-			else if (!Parser.Scanner.MatchToken(EHlslToken::Identifier))
+			else if (!Scanner.MatchToken(EHlslToken::Identifier))
 			{
-				Parser.Scanner.SourceError(TEXT("Expected Identifier!"));
+				Scanner.SourceError(TEXT("Expected Identifier!"));
 				return EParseResult::Error;
 			}
 
-			if (ParseMultiArrayBracketsAndIndex(Parser, true) == EParseResult::Error)
+			if (ParseMultiArrayBracketsAndIndex(Scanner, SymbolScope, true) == EParseResult::Error)
 			{
 				return EParseResult::Error;
 			}
 
-			if (bMultiple && Parser.Scanner.MatchToken(EHlslToken::Comma))
+			if (bMultiple && Scanner.MatchToken(EHlslToken::Comma))
 			{
 				goto NextParameter;
 			}
@@ -182,7 +182,7 @@ NextParameter:
 		}
 
 Unmatched:
-		Parser.Scanner.SetCurrentTokenIndex(OriginalToken);
+		Scanner.SetCurrentTokenIndex(OriginalToken);
 		return EParseResult::NotMatched;
 	}
 
@@ -204,51 +204,62 @@ Unmatched:
 		EDF_NOINTERPOLATION				= 0x1000,
 	};
 
-	EParseResult ParseInitializer(FHlslParser& Parser, bool bAllowLists)
+	EParseResult ParseInitializer(FHlslScanner& Scanner, FSymbolScope* SymbolScope, bool bAllowLists)
 	{
-		if (bAllowLists && Parser.Scanner.MatchToken(EHlslToken::LeftBrace))
+		if (bAllowLists && Scanner.MatchToken(EHlslToken::LeftBrace))
 		{
-			bool bFoundExpr = false;
-			while (Parser.Scanner.HasMoreTokens())
+			while (Scanner.HasMoreTokens())
 			{
-				if (Parser.Scanner.MatchToken(EHlslToken::RightBrace))
+				const auto* Token = Scanner.PeekToken();
+				if (Token->Token == EHlslToken::RightBrace)
 				{
+					Scanner.Advance();
 					return EParseResult::Matched;
 				}
-				auto Result = ParseExpression(Parser.Scanner, Parser.CurrentScope);
-				if (Result == EParseResult::Error)
+				else if (Token->Token == EHlslToken::LeftBrace)
 				{
-					Parser.Scanner.SourceError(TEXT("Invalid initielizer expression list\n"));
-					return EParseResult::Error;
+					auto Result = ParseInitializer(Scanner, SymbolScope, true);
+					if (Result != EParseResult::Matched)
+					{
+						return Result;
+					}
 				}
-				else if (Result == EParseResult::NotMatched)
+				else
 				{
-					Parser.Scanner.SourceError(TEXT("Expected initializer expression\n"));
-					return EParseResult::Error;
+					auto Result = ParseExpression(Scanner, SymbolScope);
+					if (Result == EParseResult::Error)
+					{
+						Scanner.SourceError(TEXT("Invalid initielizer expression list\n"));
+						return EParseResult::Error;
+					}
+					else if (Result == EParseResult::NotMatched)
+					{
+						Scanner.SourceError(TEXT("Expected initializer expression\n"));
+						return EParseResult::Error;
+					}
 				}
-				
-				bFoundExpr = true;
-				if (Parser.Scanner.MatchToken(EHlslToken::Comma))
+
+				if (Scanner.MatchToken(EHlslToken::Comma))
 				{
 					continue;
 				}
-				else if (Parser.Scanner.MatchToken(EHlslToken::RightBrace))
+				else if (Scanner.MatchToken(EHlslToken::RightBrace))
 				{
 					return EParseResult::Matched;
 				}
-				
-				Parser.Scanner.SourceError(TEXT("Expected ','\n"));
+
+				Scanner.SourceError(TEXT("Expected ','\n"));
 				break;
 			}
 
-			Parser.Scanner.SourceError(TEXT("Expected '}'\n"));
+			Scanner.SourceError(TEXT("Expected '}'\n"));
 		}
 		else
 		{
-			auto Result = ParseExpression(Parser.Scanner, Parser.CurrentScope);
+			auto Result = ParseExpression(Scanner, SymbolScope);
 			if (Result == EParseResult::Error)
 			{
-				Parser.Scanner.SourceError(TEXT("Invalid initializer expression\n"));
+				Scanner.SourceError(TEXT("Invalid initializer expression\n"));
 			}
 
 			return Result;
@@ -257,11 +268,9 @@ Unmatched:
 		return EParseResult::NotMatched;
 	}
 
-
-	EParseResult ParseGeneralDeclaration(FHlslParser& Parser, int32 Flags)
+	EParseResult ParseDeclarationStorageQualifiers(FHlslScanner& Scanner, int32 Flags, bool& bOutPrimitiveFound)
 	{
-		auto OriginalToken = Parser.Scanner.GetCurrentTokenIndex();
-
+		bOutPrimitiveFound = false;
 		int32 StaticFound = 0;
 		int32 SharedFound = 0;
 		int32 ConstFound = 0;
@@ -273,7 +282,7 @@ Unmatched:
 
 		if (Flags & EDF_PRIMITIVE_DATA_TYPE)
 		{
-			const auto* Token = Parser.Scanner.GetCurrentToken();
+			const auto* Token = Scanner.GetCurrentToken();
 			if (Token && Token->Token == EHlslToken::Identifier)
 			{
 				if (Token->String == TEXT("point") ||
@@ -282,77 +291,100 @@ Unmatched:
 					Token->String == TEXT("lineadj") ||
 					Token->String == TEXT("triangleadj"))
 				{
-					Parser.Scanner.Advance();
+					Scanner.Advance();
 					++PrimitiveFound;
 				}
 			}
 		}
 
-		while (Parser.Scanner.HasMoreTokens())
+		while (Scanner.HasMoreTokens())
 		{
 			bool bFound = false;
-			if ((Flags & EDF_STATIC) && Parser.Scanner.MatchToken(EHlslToken::Static))
+			if ((Flags & EDF_STATIC) && Scanner.MatchToken(EHlslToken::Static))
 			{
 				++StaticFound;
 				if (StaticFound > 1)
 				{
-					Parser.Scanner.SourceError(TEXT("'static' found more than once!\n"));
+					Scanner.SourceError(TEXT("'static' found more than once!\n"));
+					return EParseResult::Error;
 				}
 			}
-			else if ((Flags & EDF_NOINTERPOLATION) && Parser.Scanner.MatchToken(EHlslToken::NoInterpolation))
+			else if ((Flags & EDF_NOINTERPOLATION) && Scanner.MatchToken(EHlslToken::NoInterpolation))
 			{
 				++SharedFound;
 				if (SharedFound > 1)
 				{
-					Parser.Scanner.SourceError(TEXT("'groupshared' found more than once!\n"));
+					Scanner.SourceError(TEXT("'groupshared' found more than once!\n"));
+					return EParseResult::Error;
 				}
 			}
-			else if ((Flags & EDF_SHARED) && Parser.Scanner.MatchToken(EHlslToken::GroupShared))
+			else if ((Flags & EDF_SHARED) && Scanner.MatchToken(EHlslToken::GroupShared))
 			{
 				++SharedFound;
 				if (SharedFound > 1)
 				{
-					Parser.Scanner.SourceError(TEXT("'groupshared' found more than once!\n"));
+					Scanner.SourceError(TEXT("'groupshared' found more than once!\n"));
+					return EParseResult::Error;
 				}
 			}
-			else if ((Flags & EDF_CONST_ROW_MAJOR) && Parser.Scanner.MatchToken(EHlslToken::Const))
+			else if ((Flags & EDF_CONST_ROW_MAJOR) && Scanner.MatchToken(EHlslToken::Const))
 			{
 				++ConstFound;
 				if (ConstFound > 1)
 				{
-					Parser.Scanner.SourceError(TEXT("'const' found more than once!\n"));
+					Scanner.SourceError(TEXT("'const' found more than once!\n"));
+					return EParseResult::Error;
 				}
 			}
-			else if ((Flags & EDF_CONST_ROW_MAJOR) && Parser.Scanner.MatchToken(EHlslToken::RowMajor))
+			else if ((Flags & EDF_CONST_ROW_MAJOR) && Scanner.MatchToken(EHlslToken::RowMajor))
 			{
 				++RowMajorFound;
 				if (RowMajorFound > 1)
 				{
-					Parser.Scanner.SourceError(TEXT("'row_major' found more than once!\n"));
+					Scanner.SourceError(TEXT("'row_major' found more than once!\n"));
+					return EParseResult::Error;
 				}
 			}
-			else if ((Flags & EDF_IN_OUT) && Parser.Scanner.MatchToken(EHlslToken::In))
+			else if ((Flags & EDF_IN_OUT) && Scanner.MatchToken(EHlslToken::In))
 			{
 				++InFound;
 				if (InFound > 1)
 				{
-					Parser.Scanner.SourceError(TEXT("'in' found more than once!\n"));
+					Scanner.SourceError(TEXT("'in' found more than once!\n"));
+					return EParseResult::Error;
+				}
+				else if (InOutFound > 0)
+				{
+					Scanner.SourceError(TEXT("'in' can't be used with 'inout'!\n"));
+					return EParseResult::Error;
 				}
 			}
-			else if ((Flags & EDF_IN_OUT) && Parser.Scanner.MatchToken(EHlslToken::Out))
+			else if ((Flags & EDF_IN_OUT) && Scanner.MatchToken(EHlslToken::Out))
 			{
 				++OutFound;
 				if (OutFound > 1)
 				{
-					Parser.Scanner.SourceError(TEXT("'out' found more than once!\n"));
+					Scanner.SourceError(TEXT("'out' found more than once!\n"));
+					return EParseResult::Error;
+				}
+				else if (InOutFound > 0)
+				{
+					Scanner.SourceError(TEXT("'out' can't be used with 'inout'!\n"));
+					return EParseResult::Error;
 				}
 			}
-			else if ((Flags & EDF_IN_OUT) && Parser.Scanner.MatchToken(EHlslToken::InOut))
+			else if ((Flags & EDF_IN_OUT) && Scanner.MatchToken(EHlslToken::InOut))
 			{
 				++InOutFound;
 				if (InOutFound > 1)
 				{
-					Parser.Scanner.SourceError(TEXT("'inout' found more than once!\n"));
+					Scanner.SourceError(TEXT("'inout' found more than once!\n"));
+					return EParseResult::Error;
+				}
+				else if (InFound > 0 || OutFound > 0)
+				{
+					Scanner.SourceError(TEXT("'inout' can't be used with 'in' or 'out'!\n"));
+					return EParseResult::Error;
 				}
 			}
 			else
@@ -361,41 +393,59 @@ Unmatched:
 			}
 		}
 
-		bool bCanBeUnmatched = (ConstFound + RowMajorFound + InFound + OutFound + InOutFound + StaticFound + SharedFound + PrimitiveFound) == 0;
+		bOutPrimitiveFound = (PrimitiveFound > 0);
 
-		if (PrimitiveFound == 0 && (Flags & EDF_PRIMITIVE_DATA_TYPE))
+		return (ConstFound + RowMajorFound + InFound + OutFound + InOutFound + StaticFound + SharedFound + PrimitiveFound)
+			? EParseResult::Matched
+			: EParseResult::NotMatched;
+	}
+
+	EParseResult ParseStructBody(FHlslScanner& Scanner, FSymbolScope* SymbolScope);
+
+	EParseResult ParseGeneralDeclarationNoSemicolon(FHlslScanner& Scanner, FSymbolScope* SymbolScope, int32 Flags)
+	{
+		auto OriginalToken = Scanner.GetCurrentTokenIndex();
+		bool bPrimitiveFound = false;
+		auto Result = ParseDeclarationStorageQualifiers(Scanner, Flags, bPrimitiveFound);
+		if (Result == EParseResult::Error)
 		{
-			const auto* Token = Parser.Scanner.GetCurrentToken();
+			return EParseResult::Error;
+		}
+		bool bCanBeUnmatched = (Result == EParseResult::NotMatched);
+
+		if (!bPrimitiveFound && (Flags & EDF_PRIMITIVE_DATA_TYPE))
+		{
+			const auto* Token = Scanner.GetCurrentToken();
 			if (Token && Token->Token == EHlslToken::Identifier)
 			{
 				if (Token->String == TEXT("PointStream") ||
 					Token->String == TEXT("LineStream") ||
 					Token->String == TEXT("TriangleStream"))
 				{
-					Parser.Scanner.Advance();
+					Scanner.Advance();
 					bCanBeUnmatched = false;
 
-					if (!Parser.Scanner.MatchToken(EHlslToken::Lower))
+					if (!Scanner.MatchToken(EHlslToken::Lower))
 					{
-						Parser.Scanner.SourceError(TEXT("Expected '<'!"));
+						Scanner.SourceError(TEXT("Expected '<'!"));
 						return EParseResult::Error;
 					}
 
-					if (ParseGeneralType(Parser.Scanner, ETF_BUILTIN_NUMERIC | ETF_USER_TYPES, Parser.CurrentScope) != EParseResult::Matched)
+					if (ParseGeneralType(Scanner, ETF_BUILTIN_NUMERIC | ETF_USER_TYPES, SymbolScope) != EParseResult::Matched)
 					{
-						Parser.Scanner.SourceError(TEXT("Expected type!"));
+						Scanner.SourceError(TEXT("Expected type!"));
 						return EParseResult::Error;
 					}
 
-					if (!Parser.Scanner.MatchToken(EHlslToken::Greater))
+					if (!Scanner.MatchToken(EHlslToken::Greater))
 					{
-						Parser.Scanner.SourceError(TEXT("Expected '>'!"));
+						Scanner.SourceError(TEXT("Expected '>'!"));
 						return EParseResult::Error;
 					}
 
-					if (!Parser.Scanner.MatchToken(EHlslToken::Identifier))
+					if (!Scanner.MatchToken(EHlslToken::Identifier))
 					{
-						Parser.Scanner.SourceError(TEXT("Expected identifier!"));
+						Scanner.SourceError(TEXT("Expected identifier!"));
 						return EParseResult::Error;
 					}
 
@@ -404,35 +454,55 @@ Unmatched:
 			}
 		}
 
-		bool bTryBasicType = true;
 		if (Flags & EDF_TEXTURE_SAMPLER_OR_BUFFER)
 		{
-			auto Result = ParseTextureOrBufferSimpleDeclaration(Parser, (Flags & EDF_MULTIPLE) == EDF_MULTIPLE);
-			if (Result == EParseResult::Error)
+			auto Result = ParseTextureOrBufferSimpleDeclaration(Scanner, SymbolScope, (Flags & EDF_MULTIPLE) == EDF_MULTIPLE);
+			if (Result == EParseResult::Error || Result == EParseResult::Matched)
+			{
+				return Result;
+			}
+		}
+
+		if (Scanner.MatchToken(EHlslToken::Struct))
+		{
+			auto Result = ParseStructBody(Scanner, SymbolScope);
+			if (Result != EParseResult::Matched)
 			{
 				return EParseResult::Error;
 			}
-			else if (Result == EParseResult::Matched)
+
+			if (Scanner.MatchToken(EHlslToken::Identifier))
 			{
-				bTryBasicType = false;
+				//... Instance
+
+				if (Flags & EDF_INITIALIZER)
+				{
+					if (Scanner.MatchToken(EHlslToken::Equal))
+					{
+						if (ParseInitializer(Scanner, SymbolScope, (Flags & EDF_INITIALIZER_LIST) == EDF_INITIALIZER_LIST) != EParseResult::Matched)
+						{
+							Scanner.SourceError(TEXT("Invalid initializer\n"));
+							return EParseResult::Error;
+						}
+					}
+				}
 			}
 		}
-		
-		if (bTryBasicType)
+		else
 		{
-			auto Result = ParseGeneralType(Parser.Scanner, ETF_BUILTIN_NUMERIC | ETF_USER_TYPES, Parser.CurrentScope);
+			auto Result = ParseGeneralType(Scanner, ETF_BUILTIN_NUMERIC | ETF_USER_TYPES, SymbolScope);
 			if (Result == EParseResult::Matched)
 			{
 				bool bMatched = false;
 				do
 				{
-					if (!Parser.Scanner.MatchToken(EHlslToken::Identifier))
+					if (!Scanner.MatchToken(EHlslToken::Identifier))
 					{
-						Parser.Scanner.SetCurrentTokenIndex(OriginalToken);
+						Scanner.SetCurrentTokenIndex(OriginalToken);
 						return EParseResult::NotMatched;
 					}
 
-					if (ParseMultiArrayBracketsAndIndex(Parser, false) == EParseResult::Error)
+					if (ParseMultiArrayBracketsAndIndex(Scanner, SymbolScope, false) == EParseResult::Error)
 					{
 						return EParseResult::Error;
 					}
@@ -440,11 +510,11 @@ Unmatched:
 					bool bSemanticFound = false;
 					if (Flags & EDF_SEMANTIC)
 					{
-						if (Parser.Scanner.MatchToken(EHlslToken::Colon))
+						if (Scanner.MatchToken(EHlslToken::Colon))
 						{
-							if (!Parser.Scanner.MatchToken(EHlslToken::Identifier))
+							if (!Scanner.MatchToken(EHlslToken::Identifier))
 							{
-								Parser.Scanner.SourceError(TEXT("Expected identifier for semantic!"));
+								Scanner.SourceError(TEXT("Expected identifier for semantic!"));
 								return EParseResult::Error;
 							}
 
@@ -454,30 +524,43 @@ Unmatched:
 					
 					if ((Flags & EDF_INITIALIZER) && !bSemanticFound)
 					{
-						if (Parser.Scanner.MatchToken(EHlslToken::Equal))
+						if (Scanner.MatchToken(EHlslToken::Equal))
 						{
-							if (ParseInitializer(Parser, (Flags & EDF_INITIALIZER_LIST) == EDF_INITIALIZER_LIST) != EParseResult::Matched)
+							if (ParseInitializer(Scanner, SymbolScope, (Flags & EDF_INITIALIZER_LIST) == EDF_INITIALIZER_LIST) != EParseResult::Matched)
 							{
-								Parser.Scanner.SourceError(TEXT("Invalid initializer\n"));
+								Scanner.SourceError(TEXT("Invalid initializer\n"));
 								return EParseResult::Error;
 							}
 						}
 					}
 				}
-				while ((Flags & EDF_MULTIPLE) == EDF_MULTIPLE && Parser.Scanner.MatchToken(EHlslToken::Comma));
+				while ((Flags & EDF_MULTIPLE) == EDF_MULTIPLE && Scanner.MatchToken(EHlslToken::Comma));
 			}
 			else if (bCanBeUnmatched && Result == EParseResult::NotMatched)
 			{
-				Parser.Scanner.SetCurrentTokenIndex(OriginalToken);
+				Scanner.SetCurrentTokenIndex(OriginalToken);
 				return EParseResult::NotMatched;
 			}
 		}
 
+		return EParseResult::Matched;
+	}
+
+	EParseResult ParseGeneralDeclaration(FHlslScanner& Scanner, FSymbolScope* SymbolScope, int32 Flags)
+	{
+		auto OriginalToken = Scanner.GetCurrentTokenIndex();
+
+		auto Result = ParseGeneralDeclarationNoSemicolon(Scanner, SymbolScope, Flags);
+		if (Result == EParseResult::NotMatched || Result == EParseResult::Error)
+		{
+			return Result;
+		}
+
 		if (Flags & EDF_SEMICOLON)
 		{
-			if (!Parser.Scanner.MatchToken(EHlslToken::Semicolon))
+			if (!Scanner.MatchToken(EHlslToken::Semicolon))
 			{
-				Parser.Scanner.SourceError(TEXT("';' expected!\n"));
+				Scanner.SourceError(TEXT("';' expected!\n"));
 				return EParseResult::Error;
 			}
 		}
@@ -492,16 +575,20 @@ Unmatched:
 		bool bFoundRightBrace = false;
 		if (Parser.Scanner.MatchToken(EHlslToken::LeftBrace))
 		{
-			bool bFoundRightBrace = false;
 			while (Parser.Scanner.HasMoreTokens())
 			{
 				if (Parser.Scanner.MatchToken(EHlslToken::RightBrace))
 				{
-					bFoundRightBrace = true;
-					break;
+
+					if (Parser.Scanner.MatchToken(EHlslToken::Semicolon))
+					{
+						// Optional???
+					}
+
+					return EParseResult::Matched;
 				}
 
-				auto Result = ParseGeneralDeclaration(Parser, EDF_CONST_ROW_MAJOR | EDF_SEMICOLON | EDF_TEXTURE_SAMPLER_OR_BUFFER);
+				auto Result = ParseGeneralDeclaration(Parser.Scanner, Parser.CurrentScope, EDF_CONST_ROW_MAJOR | EDF_SEMICOLON | EDF_TEXTURE_SAMPLER_OR_BUFFER);
 				if (Result == EParseResult::Error)
 				{
 					return EParseResult::Error;
@@ -511,50 +598,47 @@ Unmatched:
 					break;
 				}
 			}
-
-			check(bFoundRightBrace);
-			return EParseResult::Matched;
 		}
 
 		Parser.Scanner.SourceError(TEXT("Expected '}'!"));
 		return EParseResult::Error;
 	}
 
-	EParseResult ParseStruct(FHlslParser& Parser)
+	EParseResult ParseStructBody(FHlslScanner& Scanner, FSymbolScope* SymbolScope)
 	{
 		//FInfo Info;
 		//@todo-rco: Test me!
-		const auto* Name = Parser.Scanner.GetCurrentToken();
-		if (Name && Parser.Scanner.MatchToken(EHlslToken::Identifier))
+		const auto* Name = Scanner.GetCurrentToken();
+		if (Name && Scanner.MatchToken(EHlslToken::Identifier))
 		{
-			Parser.CurrentScope->Add(Name->String);
+			SymbolScope->Add(Name->String);
 		}
 
-		if (Parser.Scanner.MatchToken(EHlslToken::Colon))
+		if (Scanner.MatchToken(EHlslToken::Colon))
 		{
-			if (!Parser.Scanner.MatchToken(EHlslToken::Identifier))
+			if (!Scanner.MatchToken(EHlslToken::Identifier))
 			{
-				Parser.Scanner.SourceError(TEXT("Identifier expected!\n"));
+				Scanner.SourceError(TEXT("Identifier expected!\n"));
 				return EParseResult::Error;
 			}
 		}
 
-		if (!Parser.Scanner.MatchToken(EHlslToken::LeftBrace))
+		if (!Scanner.MatchToken(EHlslToken::LeftBrace))
 		{
-			Parser.Scanner.SourceError(TEXT("Expected '{'!"));
+			Scanner.SourceError(TEXT("Expected '{'!"));
 			return EParseResult::Error;
 		}
 
 		bool bFoundRightBrace = false;
-		while (Parser.Scanner.HasMoreTokens())
+		while (Scanner.HasMoreTokens())
 		{
-			if (Parser.Scanner.MatchToken(EHlslToken::RightBrace))
+			if (Scanner.MatchToken(EHlslToken::RightBrace))
 			{
 				bFoundRightBrace = true;
 				break;
 			}
 
-			auto Result = ParseGeneralDeclaration(Parser, EDF_CONST_ROW_MAJOR | EDF_SEMICOLON | EDF_SEMANTIC | EDF_TEXTURE_SAMPLER_OR_BUFFER | EDF_NOINTERPOLATION);
+			auto Result = ParseGeneralDeclaration(Scanner, SymbolScope, EDF_CONST_ROW_MAJOR | EDF_SEMICOLON | EDF_SEMANTIC | EDF_TEXTURE_SAMPLER_OR_BUFFER | EDF_NOINTERPOLATION);
 			if (Result == EParseResult::Error)
 			{
 				return EParseResult::Error;
@@ -567,13 +651,7 @@ Unmatched:
 
 		if (!bFoundRightBrace)
 		{
-			Parser.Scanner.SourceError(TEXT("Expected '}'!"));
-			return EParseResult::Error;
-		}
-
-		if (!Parser.Scanner.MatchToken(EHlslToken::Semicolon))
-		{
-			Parser.Scanner.SourceError(TEXT("Expected ';'!"));
+			Scanner.SourceError(TEXT("Expected '}'!"));
 			return EParseResult::Error;
 		}
 
@@ -586,7 +664,7 @@ Unmatched:
 
 		while (Parser.Scanner.HasMoreTokens())
 		{
-			auto Result = ParseGeneralDeclaration(Parser, EDF_CONST_ROW_MAJOR | EDF_IN_OUT | EDF_TEXTURE_SAMPLER_OR_BUFFER | EDF_INITIALIZER | EDF_SEMANTIC | EDF_PRIMITIVE_DATA_TYPE | EDF_NOINTERPOLATION);
+			auto Result = ParseGeneralDeclaration(Parser.Scanner, Parser.CurrentScope, EDF_CONST_ROW_MAJOR | EDF_IN_OUT | EDF_TEXTURE_SAMPLER_OR_BUFFER | EDF_INITIALIZER | EDF_SEMANTIC | EDF_PRIMITIVE_DATA_TYPE | EDF_NOINTERPOLATION);
 			if (Result == EParseResult::Error)
 			{
 				return EParseResult::Error;
@@ -670,6 +748,7 @@ Unmatched:
 
 	EParseResult ParseStatementBlock(FHlslParser& Parser)
 	{
+		FCreateSymbolScope SymbolScope(&Parser.CurrentScope);
 		while (Parser.Scanner.HasMoreTokens())
 		{
 			auto Result = ParseStatement(Parser);
@@ -734,12 +813,12 @@ Unmatched:
 
 	EParseResult ParseLocalDeclaration(FHlslParser& Parser)
 	{
-		return ParseGeneralDeclaration(Parser, EDF_CONST_ROW_MAJOR | EDF_INITIALIZER | EDF_INITIALIZER_LIST | EDF_SEMICOLON | EDF_MULTIPLE);
+		return ParseGeneralDeclaration(Parser.Scanner, Parser.CurrentScope, EDF_CONST_ROW_MAJOR | EDF_INITIALIZER | EDF_INITIALIZER_LIST | EDF_SEMICOLON | EDF_MULTIPLE);
 	}
 
 	EParseResult ParseGlobalVariableDeclaration(FHlslParser& Parser)
 	{
-		return ParseGeneralDeclaration(Parser, EDF_CONST_ROW_MAJOR | EDF_STATIC | EDF_SHARED | EDF_TEXTURE_SAMPLER_OR_BUFFER | EDF_INITIALIZER | EDF_INITIALIZER_LIST | EDF_SEMICOLON | EDF_MULTIPLE);
+		return ParseGeneralDeclaration(Parser.Scanner, Parser.CurrentScope, EDF_CONST_ROW_MAJOR | EDF_STATIC | EDF_SHARED | EDF_TEXTURE_SAMPLER_OR_BUFFER | EDF_INITIALIZER | EDF_INITIALIZER_LIST | EDF_SEMICOLON | EDF_MULTIPLE);
 	}
 
 	EParseResult ParseReturnStatement(FHlslParser& Parser)
@@ -801,6 +880,8 @@ Unmatched:
 
 	EParseResult ParseWhileStatement(FHlslParser& Parser)
 	{
+		FCreateSymbolScope SymbolScope(&Parser.CurrentScope);
+
 		if (!Parser.Scanner.MatchToken(EHlslToken::LeftParenthesis))
 		{
 			Parser.Scanner.SourceError(TEXT("'(' expected"));
@@ -830,13 +911,15 @@ Unmatched:
 
 	EParseResult ParseForStatement(FHlslParser& Parser)
 	{
+		FCreateSymbolScope SymbolScope(&Parser.CurrentScope);
+
 		if (!Parser.Scanner.MatchToken(EHlslToken::LeftParenthesis))
 		{
 			Parser.Scanner.SourceError(TEXT("Expected '('!\n"));
 			return EParseResult::Error;
 		}
 
-		auto Result = ParseGeneralDeclaration(Parser, EDF_CONST_ROW_MAJOR | EDF_INITIALIZER);
+		auto Result = ParseGeneralDeclaration(Parser.Scanner, Parser.CurrentScope, EDF_CONST_ROW_MAJOR | EDF_INITIALIZER);
 		if (Result == EParseResult::Error)
 		{
 			Parser.Scanner.SourceError(TEXT("Expected expression or declaration!\n"));
@@ -889,6 +972,8 @@ Unmatched:
 
 	EParseResult ParseIfStatement(FHlslParser& Parser)
 	{
+		FCreateSymbolScope SymbolScope(&Parser.CurrentScope);
+
 		if (!Parser.Scanner.MatchToken(EHlslToken::LeftParenthesis))
 		{
 			Parser.Scanner.SourceError(TEXT("'(' expected"));
@@ -1101,7 +1186,6 @@ Unmatched:
 		{
 			bInit = true;
 
-			RulesTranslationUnit.Add(FRulePair(EHlslToken::Struct, ParseStruct));
 			RulesTranslationUnit.Add(FRulePair(EHlslToken::CBuffer, ParseCBuffer));
 			RulesTranslationUnit.Add(FRulePair(EHlslToken::Invalid, ParseFunctionDeclaration, TryParseAttribute));
 			RulesTranslationUnit.Add(FRulePair(EHlslToken::Invalid, ParseGlobalVariableDeclaration));
diff --git a/Engine/Source/Programs/CrossCompilerTool/Private/CrossCompilerTool.cpp b/Engine/Source/Programs/CrossCompilerTool/Private/CrossCompilerTool.cpp
--- a/Engine/Source/Programs/CrossCompilerTool/Private/CrossCompilerTool.cpp
+++ b/Engine/Source/Programs/CrossCompilerTool/Private/CrossCompilerTool.cpp
@@ -22,7 +22,7 @@ namespace CCT
 		ILanguageSpec* Language = nullptr;
 		FCodeBackend* Backend = nullptr;
 
-		int32 Flags = 0;
+		uint32 Flags = 0;
 
 		Flags |= RunInfo.bRunCPP ? 0 : HLSLCC_NoPreprocess;
 
@@ -81,12 +81,20 @@ namespace CCT
 					}
 					UE_LOG(LogCrossCompilerTool, Log, TEXT("%s!"), *File);
 
-					CrossCompiler::Parser::Parse(HLSLShader, File);
+					if (!CrossCompiler::Parser::Parse(HLSLShader, File))
+					{
+						UE_LOG(LogCrossCompilerTool, Log, TEXT("Error compiling '%s'!"), *File);
+						return 0;
+					}
 				}
 			}
 			else
 			{
-				CrossCompiler::Parser::Parse(HLSLShaderSource, *RunInfo.InputFile);
+				if (!CrossCompiler::Parser::Parse(HLSLShaderSource, *RunInfo.InputFile))
+				{
+					UE_LOG(LogCrossCompilerTool, Log, TEXT("Error compiling '%s'!"), *RunInfo.InputFile);
+					return 0;
+				}
 			}
 			//Scanner.Dump();
 			return 1;
