UE4 - HlslParser prep changes

diff --git a/Engine/Source/Developer/ShaderCompilerCommon/Private/HlslExpressionParser.inl b/Engine/Source/Developer/ShaderCompilerCommon/Private/HlslExpressionParser.inl
--- a/Engine/Source/Developer/ShaderCompilerCommon/Private/HlslExpressionParser.inl
+++ b/Engine/Source/Developer/ShaderCompilerCommon/Private/HlslExpressionParser.inl
@@ -4,6 +4,7 @@
 	HlslExpressionParser.inl - Implementation for parsing hlsl expressions.
 =============================================================================*/
 
+#pragma once
 #include "ShaderCompilerCommon.h"
 #include "HlslParser.h"
 
@@ -331,7 +332,7 @@ namespace CrossCompiler
 			{
 				return EParseResult::Error;
 			}
-			Info.PrintWithTabs(FString::Printf(TEXT("Atom %s\n"), *Token->String));
+			//Info.PrintWithTabs(FString::Printf(TEXT("Atom %s\n"), *Token->String));
 
 			bool bFoundUnary = true;
 			switch (Token->Token)
@@ -367,6 +368,11 @@ namespace CrossCompiler
 				break;
 
 			case EHlslToken::BoolConstant:
+				Scanner.Advance();
+				bFoundUnary = false;
+				bFoundNonUnary = true;
+				break;
+
 			case EHlslToken::UnsignedIntegerConstant:
 				//FPlatformMisc::LowLevelOutputDebugStringf(TEXT("Atom Int %d\n"), Token->UnsignedInteger);
 				Scanner.Advance();
@@ -396,7 +402,7 @@ namespace CrossCompiler
 				if (Peek1 && ParseGeneralType(Peek1, ETF_BUILTIN_NUMERIC | ETF_USER_TYPES, SymbolScope) == EParseResult::Matched && Peek2 && Peek2->Token == EHlslToken::RightParenthesis)
 				{
 					// Cast
-					Info.PrintWithTabs(FString::Printf(TEXT("Cast %s\n"), *Peek1->String));
+					//Info.PrintWithTabs(FString::Printf(TEXT("Cast %s\n"), *Peek1->String));
 					Scanner.Advance();
 					Scanner.Advance();
 				}
@@ -447,42 +453,42 @@ namespace CrossCompiler
 				// Grrr handle Sampler as a variable name
 				if (ParseGeneralType(Scanner, ETF_SAMPLER_TEXTURE_BUFFER) == EParseResult::Matched)
 				{
-					bFoundUnary = false;
+				bFoundUnary = false;
 					bFoundNonUnary = true;
-					break;
-				}
+				break;
+			}
 				// Handle float3(x,y,z)
 				else if (ParseGeneralType(Scanner, ETF_BUILTIN_NUMERIC) == EParseResult::Matched)
-				{
+			{
 					bFoundUnary = false;
 					//bFoundTypeInitializerList = true;
 					bFoundNonUnary = true;
-					break;
-				}
+				break;
+			}
 				else
-				{
+	{
 					return EParseResult::NotMatched;
-				}
-			}
+		}
+		}
 
 			if (!bFoundUnary)
-			{
-				break;
-			}
+		{
+			break;
+		}
 		}
 		while (Scanner.HasMoreTokens());
 
 		if (!bFoundNonUnary)
-		{
+			{
 			Scanner.SourceError(TEXT("Expected expression!"));
-			return EParseResult::Error;
-		}
+				return EParseResult::Error;
+			}
 
 		// Suffix
 		while (Scanner.HasMoreTokens())
 		{
 			auto* Token = Scanner.GetCurrentToken();
-			Info.PrintWithTabs(FString::Printf(TEXT("Atom %s\n"), *Token->String));
+			//Info.PrintWithTabs(FString::Printf(TEXT("Atom %s\n"), *Token->String));
 			if (Scanner.MatchToken(EHlslToken::LeftSquareBracket))
 			{
 				//FPlatformMisc::LowLevelOutputDebugStringf(TEXT("Array\n"));
@@ -510,7 +516,7 @@ namespace CrossCompiler
 			}
 			else if (Scanner.MatchToken(EHlslToken::LeftParenthesis))
 			{
-				Info.PrintWithTabs(TEXT("FunctionCall (\n"));
+				//Info.PrintWithTabs(TEXT("FunctionCall (\n"));
 				bool bFirst = true;
 				bool bFoundRightParenthesis = false;
 				while (Scanner.HasMoreTokens())
diff --git a/Engine/Source/Developer/ShaderCompilerCommon/Private/HlslLexer.cpp b/Engine/Source/Developer/ShaderCompilerCommon/Private/HlslLexer.cpp
--- a/Engine/Source/Developer/ShaderCompilerCommon/Private/HlslLexer.cpp
+++ b/Engine/Source/Developer/ShaderCompilerCommon/Private/HlslLexer.cpp
@@ -1,5 +1,8 @@
 // Copyright 1998-2014 Epic Games, Inc. All Rights Reserved.
-// HlslLexer.cpp - Implementation for scanning & tokenizing hlsl
+
+/*=============================================================================
+	HlslLexer.cpp - Implementation for scanning & tokenizing hlsl
+=============================================================================*/
 
 
 #include "ShaderCompilerCommon.h"
@@ -886,9 +889,9 @@ namespace CrossCompiler
 	inline void FHlslScanner::AddToken(const FHlslToken& Token, const FTokenizer& Tokenizer)
 	{
 		int32 TokenIndex = Tokens.Add(Token);
-		Tokens[TokenIndex].SourceFilename = &SourceFilenames.Last();
-		Tokens[TokenIndex].SourceLine = Tokenizer.Line;
-		Tokens[TokenIndex].SourceColumn = (int32)(Tokenizer.Current - Tokenizer.CurrentLineStart) + 1;
+		Tokens[TokenIndex].SourceInfo.Filename = &SourceFilenames.Last();
+		Tokens[TokenIndex].SourceInfo.Line = Tokenizer.Line;
+		Tokens[TokenIndex].SourceInfo.Column = (int32)(Tokenizer.Current - Tokenizer.CurrentLineStart) + 1;
 	}
 
 	void FHlslScanner::Clear(const FString& Filename)
@@ -1059,8 +1062,8 @@ namespace CrossCompiler
 		if (CurrentToken < (uint32)Tokens.Num())
 		{
 			const auto& Token = Tokens[CurrentToken];
-			check(Token.SourceFilename);
-			FPlatformMisc::LowLevelOutputDebugStringf(TEXT("%s(%d): (%d) %s\n"), **Token.SourceFilename, Token.SourceLine, Token.SourceColumn, *Error);
+			check(Token.SourceInfo.Filename);
+			FPlatformMisc::LowLevelOutputDebugStringf(TEXT("%s(%d): (%d) %s\n"), **Token.SourceInfo.Filename, Token.SourceInfo.Line, Token.SourceInfo.Column, *Error);
 		}
 		else
 		{
diff --git a/Engine/Source/Developer/ShaderCompilerCommon/Private/HlslLexer.h b/Engine/Source/Developer/ShaderCompilerCommon/Private/HlslLexer.h
--- a/Engine/Source/Developer/ShaderCompilerCommon/Private/HlslLexer.h
+++ b/Engine/Source/Developer/ShaderCompilerCommon/Private/HlslLexer.h
@@ -246,6 +246,15 @@ namespace CrossCompiler
 		BoolConstant,
 	};
 
+	struct FSourceInfo
+	{
+		FString* Filename;
+		int32 Line;
+		int32 Column;
+
+		FSourceInfo() : Filename(nullptr), Line(0), Column(0) {}
+	};
+
 	struct FHlslToken
 	{
 		EHlslToken Token;
@@ -253,15 +262,13 @@ namespace CrossCompiler
 		uint32 UnsignedInteger;
 		float Float;
 
-		FString* SourceFilename;
-		int32 SourceLine;
-		int32 SourceColumn;
+		FSourceInfo SourceInfo;
 
-		explicit FHlslToken(const FString& Identifier) : Token(EHlslToken::Identifier), String(Identifier), UnsignedInteger(0), Float(0), SourceFilename(nullptr), SourceLine(0), SourceColumn(0) { }
-		explicit FHlslToken(EHlslToken InToken, const FString& Identifier) : Token(InToken), String(Identifier), UnsignedInteger(0), Float(0), SourceFilename(nullptr), SourceLine(0), SourceColumn(0) { }
-		explicit FHlslToken(uint32 InUnsignedInteger) : Token(EHlslToken::UnsignedIntegerConstant), UnsignedInteger(InUnsignedInteger), Float(0), SourceFilename(nullptr), SourceLine(0), SourceColumn(0) { }
-		explicit FHlslToken(float InFloat) : Token(EHlslToken::FloatConstant), UnsignedInteger(0), Float(InFloat), SourceFilename(nullptr), SourceLine(0), SourceColumn(0) { }
-		explicit FHlslToken(bool bInValue) : Token(EHlslToken::BoolConstant), UnsignedInteger(bInValue ? 1 : 0), Float(0), SourceFilename(nullptr), SourceLine(0), SourceColumn(0) { }
+		explicit FHlslToken(const FString& Identifier) : Token(EHlslToken::Identifier), String(Identifier), UnsignedInteger(0), Float(0) { }
+		explicit FHlslToken(EHlslToken InToken, const FString& Identifier) : Token(InToken), String(Identifier), UnsignedInteger(0), Float(0) { }
+		explicit FHlslToken(uint32 InUnsignedInteger) : Token(EHlslToken::UnsignedIntegerConstant), UnsignedInteger(InUnsignedInteger), Float(0) { }
+		explicit FHlslToken(float InFloat) : Token(EHlslToken::FloatConstant), UnsignedInteger(0), Float(InFloat) { }
+		explicit FHlslToken(bool bInValue) : Token(EHlslToken::BoolConstant), UnsignedInteger(bInValue ? 1 : 0), Float(0) { }
 	};
 
 	class FHlslScanner
diff --git a/Engine/Source/Developer/ShaderCompilerCommon/Private/HlslParser.h b/Engine/Source/Developer/ShaderCompilerCommon/Private/HlslParser.h
--- a/Engine/Source/Developer/ShaderCompilerCommon/Private/HlslParser.h
+++ b/Engine/Source/Developer/ShaderCompilerCommon/Private/HlslParser.h
@@ -19,6 +19,7 @@ namespace CrossCompiler
 
 	namespace Parser
 	{
+		// Returns true if successfully parsed
 		bool Parse(const FString& Input, const FString& Filename);
 	}
 }
