UE4 - HlslParser
- Added ExpressionList parser
- Split out unary operator parsing

diff --git a/Engine/Source/Developer/ShaderCompilerCommon/Private/HlslExpressionParser.inl b/Engine/Source/Developer/ShaderCompilerCommon/Private/HlslExpressionParser.inl
--- a/Engine/Source/Developer/ShaderCompilerCommon/Private/HlslExpressionParser.inl
+++ b/Engine/Source/Developer/ShaderCompilerCommon/Private/HlslExpressionParser.inl
@@ -10,6 +10,12 @@
 
 namespace CrossCompiler
 {
+	struct FSymbolScope;
+	struct FInfo;
+
+	EParseResult ComputeExpr(FHlslScanner& Scanner, int32 MinPrec, FInfo& Info, FSymbolScope* SymbolScope);
+	EParseResult ParseExpressionList(EHlslToken EndListToken, FHlslScanner& Scanner, FSymbolScope* SymbolScope, EHlslToken NewStartListToken = EHlslToken::Invalid);
+
 	struct FSymbolScope
 	{
 		FSymbolScope* Parent;
@@ -318,171 +324,160 @@ namespace CrossCompiler
 		return EParseResult::NotMatched;
 	}
 
-	EParseResult ComputeExpr(FHlslScanner& Scanner, int32 MinPrec, FInfo& Info, FSymbolScope* SymbolScope);
-
-	EParseResult ComputeAtom(FHlslScanner& Scanner, FInfo& Info, FSymbolScope* SymbolScope)
+	EParseResult MatchUnaryOperator(FHlslScanner& Scanner, FInfo& Info, FSymbolScope* SymbolScope)
 	{
-		// Prefix (unary)
-		bool bFoundNonUnary = false;
-		//bool bFoundTypeInitializerList = false;
-		do
+		bool bFoundAny = false;
+		bool bTryAgain = true;
+		while (Scanner.HasMoreTokens() && bTryAgain)
 		{
 			auto* Token = Scanner.GetCurrentToken();
-			if (!Token)
-			{
-				return EParseResult::Error;
-			}
-			//Info.PrintWithTabs(FString::Printf(TEXT("Atom %s\n"), *Token->String));
-
-			bool bFoundUnary = true;
 			switch (Token->Token)
 			{
 			case EHlslToken::PlusPlus:
-				//FPlatformMisc::LowLevelOutputDebugStringf(TEXT("Unary %s\n"), *Token->String);
+				bFoundAny = true;
 				Scanner.Advance();
 				break;
 
 			case EHlslToken::MinusMinus:
-				//FPlatformMisc::LowLevelOutputDebugStringf(TEXT("Unary %s\n"), *Token->String);
+				bFoundAny = true;
 				Scanner.Advance();
 				break;
 
 			case EHlslToken::Plus:
-				//FPlatformMisc::LowLevelOutputDebugStringf(TEXT("Unary %s\n"), *Token->String);
+				bFoundAny = true;
 				Scanner.Advance();
 				break;
 
 			case EHlslToken::Minus:
-				//FPlatformMisc::LowLevelOutputDebugStringf(TEXT("Unary %s\n"), *Token->String);
+				bFoundAny = true;
 				Scanner.Advance();
 				break;
 
 			case EHlslToken::Not:
-				//FPlatformMisc::LowLevelOutputDebugStringf(TEXT("Unary %s\n"), *Token->String);
+				bFoundAny = true;
 				Scanner.Advance();
 				break;
 
 			case EHlslToken::Neg:
-				//FPlatformMisc::LowLevelOutputDebugStringf(TEXT("Unary %s\n"), *Token->String);
-				Scanner.Advance();
-				break;
-
-			case EHlslToken::BoolConstant:
-				Scanner.Advance();
-				bFoundUnary = false;
-				bFoundNonUnary = true;
-				break;
-
-			case EHlslToken::UnsignedIntegerConstant:
-				//FPlatformMisc::LowLevelOutputDebugStringf(TEXT("Atom Int %d\n"), Token->UnsignedInteger);
-				Scanner.Advance();
-				bFoundUnary = false;
-				bFoundNonUnary = true;
-				break;
-
-			case EHlslToken::FloatConstant:
-				//FPlatformMisc::LowLevelOutputDebugStringf(TEXT("Atom Float %f\n"), Token->Float);
+				bFoundAny = true;
 				Scanner.Advance();
-				bFoundUnary = false;
-				bFoundNonUnary = true;
 				break;
 
-			case EHlslToken::Identifier:
-				//FPlatformMisc::LowLevelOutputDebugStringf(TEXT("Atom Literal %s\n"), *Token->String);
-				Scanner.Advance();
-				bFoundUnary = false;
-				bFoundNonUnary = true;
-				break;
 			case EHlslToken::LeftParenthesis:
+			// Only cast expressions are Unary
 			{
-				Scanner.Advance();
-
-				const auto* Peek1 = Scanner.PeekToken(0);
-				const auto* Peek2 = Scanner.PeekToken(1);
+				const auto* Peek1 = Scanner.PeekToken(1);
+				const auto* Peek2 = Scanner.PeekToken(2);
 				if (Peek1 && ParseGeneralType(Peek1, ETF_BUILTIN_NUMERIC | ETF_USER_TYPES, SymbolScope) == EParseResult::Matched && Peek2 && Peek2->Token == EHlslToken::RightParenthesis)
 				{
 					// Cast
-					//Info.PrintWithTabs(FString::Printf(TEXT("Cast %s\n"), *Peek1->String));
 					Scanner.Advance();
 					Scanner.Advance();
+					Scanner.Advance();
+					bFoundAny = true;
 				}
 				else
 				{
-					if (ComputeExpr(Scanner, 1, Info, SymbolScope) != EParseResult::Matched)
-					{
-						Scanner.SourceError(TEXT("Expected expression!"));
-						return EParseResult::Error;
-					}
-
-					if (!Scanner.MatchToken(EHlslToken::RightParenthesis))
-					{
-						Scanner.SourceError(TEXT("Expected ')'!"));
-						return EParseResult::Error;
-					}
-
-					bFoundNonUnary = true;
-					bFoundUnary = false;
+					// Non-unary
+					return bFoundAny ? EParseResult::Matched : EParseResult::NotMatched;
 				}
+			}
 				break;
-/*
-				//FPlatformMisc::LowLevelOutputDebugStringf(TEXT("Atom (\n"));
 
-				// Try cast expression
-				else
+			default:
+				return bFoundAny ? EParseResult::Matched : EParseResult::NotMatched;
+			}
+		}
+
+		// Ran out of tokens!
+		return EParseResult::Error;
+	}
+
+	EParseResult ComputeAtom(FHlslScanner& Scanner, FInfo& Info, FSymbolScope* SymbolScope)
+	{
+		auto UnaryResult = MatchUnaryOperator(Scanner, Info, SymbolScope);
+		auto* Token = Scanner.GetCurrentToken();
+		if (!Token || UnaryResult == EParseResult::Error)
+		{
+			return EParseResult::Error;
+		}
+
+		switch (Token->Token)
+		{
+		case EHlslToken::BoolConstant:
+			Scanner.Advance();
+			break;
+
+		case EHlslToken::UnsignedIntegerConstant:
+			Scanner.Advance();
+			break;
+
+		case EHlslToken::FloatConstant:
+			Scanner.Advance();
+			break;
+
+		case EHlslToken::Identifier:
+			Scanner.Advance();
+			break;
+
+		case EHlslToken::LeftParenthesis:
+		{
+			Scanner.Advance();
+
+			// Check if it's a cast expression first
+			const auto* Peek1 = Scanner.PeekToken(0);
+			const auto* Peek2 = Scanner.PeekToken(1);
+			// Parenthesis expression
+			if (ComputeExpr(Scanner, 1, Info, SymbolScope) != EParseResult::Matched)
+			{
+				Scanner.SourceError(TEXT("Expected expression!"));
+				return EParseResult::Error;
+			}
+
+			if (!Scanner.MatchToken(EHlslToken::RightParenthesis))
+			{
+				Scanner.SourceError(TEXT("Expected ')'!"));
+				return EParseResult::Error;
+			}
+		}
+			break;
+
+		default:
+			// Grrr handle Sampler as a variable name... This is safe here since Declarations are always handled first
+			if (ParseGeneralType(Scanner, ETF_SAMPLER_TEXTURE_BUFFER) == EParseResult::Matched)
+			{
+				break;
+			}
+			// Handle float3(x,y,z)
+			else if (ParseGeneralType(Scanner, ETF_BUILTIN_NUMERIC) == EParseResult::Matched)
+			{
+				if (Scanner.MatchToken(EHlslToken::LeftParenthesis))
 				{
-					auto Result = ComputeExpr(Scanner, 1, Info, SymbolScope);
+					auto Result = ParseExpressionList(EHlslToken::RightParenthesis, Scanner, SymbolScope);
 					if (Result != EParseResult::Matched)
 					{
-						Scanner.SourceError(TEXT("Expected expression!"));
+						Scanner.SourceError(TEXT("Unexpected type in numeric constructor!"));
 						return EParseResult::Error;
 					}
-
-					if (!Scanner.MatchToken(EHlslToken::RightParenthesis))
-					{
-						Scanner.SourceError(TEXT("Expected ')'!"));
-						return EParseResult::Error;
-					}
-
-					bFoundNonUnary = true;
-					bFoundUnary = false;
 				}
-*/
-			}
-				break;
-			default:
-				// Grrr handle Sampler as a variable name
-				if (ParseGeneralType(Scanner, ETF_SAMPLER_TEXTURE_BUFFER) == EParseResult::Matched)
+				else
 				{
-				bFoundUnary = false;
-					bFoundNonUnary = true;
+					Scanner.SourceError(TEXT("Unexpected type in declaration!"));
+					return EParseResult::Error;
+				}
 				break;
 			}
-				// Handle float3(x,y,z)
-				else if (ParseGeneralType(Scanner, ETF_BUILTIN_NUMERIC) == EParseResult::Matched)
+			else
 			{
-					bFoundUnary = false;
-					//bFoundTypeInitializerList = true;
-					bFoundNonUnary = true;
-				break;
-			}
-				else
-	{
-					return EParseResult::NotMatched;
-		}
-		}
-
-			if (!bFoundUnary)
-		{
-			break;
-		}
-		}
-		while (Scanner.HasMoreTokens());
+				if (UnaryResult == EParseResult::Matched)
+				{
+					Scanner.SourceError(TEXT("Expected expression!"));
+					return EParseResult::Error;
+				}
 
-		if (!bFoundNonUnary)
-			{
-			Scanner.SourceError(TEXT("Expected expression!"));
-				return EParseResult::Error;
+				return EParseResult::NotMatched;
 			}
+		}
 
 		// Suffix
 		while (Scanner.HasMoreTokens())
@@ -516,38 +511,9 @@ namespace CrossCompiler
 			}
 			else if (Scanner.MatchToken(EHlslToken::LeftParenthesis))
 			{
-				//Info.PrintWithTabs(TEXT("FunctionCall (\n"));
-				bool bFirst = true;
-				bool bFoundRightParenthesis = false;
-				while (Scanner.HasMoreTokens())
-				{
-					if (Scanner.MatchToken(EHlslToken::RightParenthesis))
-					{
-						bFoundRightParenthesis = true;
-						break;
-					}
-
-					if (bFirst)
-					{
-						bFirst = false;
-					}
-					else
-					{
-						if (!Scanner.MatchToken(EHlslToken::Comma))
-						{
-							Scanner.SourceError(TEXT("Expected ','!"));
-							return EParseResult::Error;
-						}
-					}
-					auto Result = ComputeExpr(Scanner, 0, Info, SymbolScope);
-					if (Result != EParseResult::Matched)
-					{
-						Scanner.SourceError(TEXT("Expected expression!"));
-						return EParseResult::Error;
-					}
-				}
-
-				if (!bFoundRightParenthesis)
+				// Function Call
+				auto Result = ParseExpressionList(EHlslToken::RightParenthesis, Scanner, SymbolScope);
+				if (Result != EParseResult::Matched)
 				{
 					Scanner.SourceError(TEXT("Expected ')'!"));
 					return EParseResult::Error;
@@ -724,4 +690,54 @@ namespace CrossCompiler
 		FInfo Info(!true);
 		return ComputeExpr(Scanner, 0, Info, SymbolScope);
 	}
+
+	EParseResult ParseExpressionList(EHlslToken EndListToken, FHlslScanner& Scanner, FSymbolScope* SymbolScope, EHlslToken NewStartListToken)
+	{
+		while (Scanner.HasMoreTokens())
+		{
+			const auto* Token = Scanner.PeekToken();
+			if (Token->Token == EndListToken)
+			{
+				Scanner.Advance();
+				return EParseResult::Matched;
+			}
+			else if (NewStartListToken != EHlslToken::Invalid && Token->Token == NewStartListToken)
+			{
+				Scanner.Advance();
+				auto Result = ParseExpressionList(EndListToken, Scanner, SymbolScope, NewStartListToken);
+				if (Result != EParseResult::Matched)
+				{
+					return Result;
+				}
+			}
+			else
+			{
+				auto Result = ParseExpression(Scanner, SymbolScope);
+				if (Result == EParseResult::Error)
+				{
+					Scanner.SourceError(TEXT("Invalid expression list\n"));
+					return EParseResult::Error;
+				}
+				else if (Result == EParseResult::NotMatched)
+				{
+					Scanner.SourceError(TEXT("Expected expression\n"));
+					return EParseResult::Error;
+				}
+			}
+
+			if (Scanner.MatchToken(EHlslToken::Comma))
+			{
+				continue;
+			}
+			else if (Scanner.MatchToken(EndListToken))
+			{
+				return EParseResult::Matched;
+			}
+
+			Scanner.SourceError(TEXT("Expected ','\n"));
+			break;
+		}
+
+		return EParseResult::Error;
+	}
 }
diff --git a/Engine/Source/Developer/ShaderCompilerCommon/Private/HlslParser.cpp b/Engine/Source/Developer/ShaderCompilerCommon/Private/HlslParser.cpp
--- a/Engine/Source/Developer/ShaderCompilerCommon/Private/HlslParser.cpp
+++ b/Engine/Source/Developer/ShaderCompilerCommon/Private/HlslParser.cpp
@@ -208,51 +208,13 @@ Unmatched:
 	{
 		if (bAllowLists && Scanner.MatchToken(EHlslToken::LeftBrace))
 		{
-			while (Scanner.HasMoreTokens())
+			auto Result = ParseExpressionList(EHlslToken::RightBrace, Scanner, SymbolScope, EHlslToken::LeftBrace);
+			if (Result != EParseResult::Matched)
 			{
-				const auto* Token = Scanner.PeekToken();
-				if (Token->Token == EHlslToken::RightBrace)
-				{
-					Scanner.Advance();
-					return EParseResult::Matched;
-				}
-				else if (Token->Token == EHlslToken::LeftBrace)
-				{
-					auto Result = ParseInitializer(Scanner, SymbolScope, true);
-					if (Result != EParseResult::Matched)
-					{
-						return Result;
-					}
-				}
-				else
-				{
-					auto Result = ParseExpression(Scanner, SymbolScope);
-					if (Result == EParseResult::Error)
-					{
-						Scanner.SourceError(TEXT("Invalid initielizer expression list\n"));
-						return EParseResult::Error;
-					}
-					else if (Result == EParseResult::NotMatched)
-					{
-						Scanner.SourceError(TEXT("Expected initializer expression\n"));
-						return EParseResult::Error;
-					}
-				}
-
-				if (Scanner.MatchToken(EHlslToken::Comma))
-				{
-					continue;
-				}
-				else if (Scanner.MatchToken(EHlslToken::RightBrace))
-				{
-					return EParseResult::Matched;
-				}
-
-				Scanner.SourceError(TEXT("Expected ','\n"));
-				break;
+				Scanner.SourceError(TEXT("Invalid initializer list\n"));
 			}
 
-			Scanner.SourceError(TEXT("Expected '}'\n"));
+			return EParseResult::Matched;
 		}
 		else
 		{
@@ -1022,22 +984,8 @@ Unmatched:
 
 			if (Parser.Scanner.MatchToken(EHlslToken::LeftParenthesis))
 			{
-				do
-				{
-					if (ParseExpression(Parser.Scanner, Parser.CurrentScope) != EParseResult::Matched)
-					{
-						Parser.Scanner.SourceError(TEXT("Incorrect attribute! Expected expression.\n"));
-						return EParseResult::Error;
-					}
-
-					if (!Parser.Scanner.MatchToken(EHlslToken::Comma))
-					{
-						break;
-					}
-				}
-				while (Parser.Scanner.HasMoreTokens());
-
-				if (!Parser.Scanner.MatchToken(EHlslToken::RightParenthesis))
+				auto Result = ParseExpressionList(EHlslToken::RightParenthesis, Parser.Scanner, Parser.CurrentScope);
+				if (Result != EParseResult::Matched)
 				{
 					Parser.Scanner.SourceError(TEXT("Incorrect attribute! Expected ')'.\n"));
 					return EParseResult::Error;
@@ -1199,8 +1147,9 @@ Unmatched:
 			RulesStatements.Add(FRulePair(EHlslToken::Switch, ParseSwitchStatement, TryParseAttribute));
 			RulesStatements.Add(FRulePair(EHlslToken::Semicolon, ParseEmptyStatement));
 			RulesStatements.Add(FRulePair(EHlslToken::Break, ParseEmptyStatement));
-			RulesStatements.Add(FRulePair(EHlslToken::Invalid, ParseExpressionStatement));
 			RulesStatements.Add(FRulePair(EHlslToken::Invalid, ParseLocalDeclaration));
+			// Always try expressions last
+			RulesStatements.Add(FRulePair(EHlslToken::Invalid, ParseExpressionStatement));
 		}
 	}
 
diff --git a/Engine/Source/Programs/CrossCompilerTool/Private/CrossCompilerTool.cpp b/Engine/Source/Programs/CrossCompilerTool/Private/CrossCompilerTool.cpp
--- a/Engine/Source/Programs/CrossCompilerTool/Private/CrossCompilerTool.cpp
+++ b/Engine/Source/Programs/CrossCompilerTool/Private/CrossCompilerTool.cpp
@@ -65,26 +65,27 @@ namespace CCT
 			if (RunInfo.bList)
 			{
 				TArray<FString> List;
-				FString Temp;
-				while (HLSLShaderSource.Split(TEXT("\r\n"), &Temp, &HLSLShaderSource))
+
+				if (!FFileHelper::LoadANSITextFileToStrings(*RunInfo.InputFile, &IFileManager::Get(), List))
 				{
-					List.Add(Temp);
+					return 1;
 				}
 
+				int32 Count = 0;
 				for (auto& File : List)
 				{
 					FString HLSLShader;
 					if (!FFileHelper::LoadFileToString(HLSLShader, *File))
 					{
-						UE_LOG(LogCrossCompilerTool, Error, TEXT("Couldn't load Input file '%s'!"), *RunInfo.InputFile);
-						return 1;
+						UE_LOG(LogCrossCompilerTool, Error, TEXT("Couldn't load Input file '%s'!"), *File);
+						continue;
 					}
-					UE_LOG(LogCrossCompilerTool, Log, TEXT("%s!"), *File);
+					UE_LOG(LogCrossCompilerTool, Log, TEXT("%d: %s!"), Count++, *File);
 
 					if (!CrossCompiler::Parser::Parse(HLSLShader, File))
 					{
 						UE_LOG(LogCrossCompilerTool, Log, TEXT("Error compiling '%s'!"), *File);
-						return 0;
+						return 1;
 					}
 				}
 			}
@@ -93,11 +94,11 @@ namespace CCT
 				if (!CrossCompiler::Parser::Parse(HLSLShaderSource, *RunInfo.InputFile))
 				{
 					UE_LOG(LogCrossCompilerTool, Log, TEXT("Error compiling '%s'!"), *RunInfo.InputFile);
-					return 0;
+					return 1;
 				}
 			}
 			//Scanner.Dump();
-			return 1;
+			return 0;
 		}
 
 		ANSICHAR* ShaderSource = 0;
