UE4 - HlslParser - numeric expressions AST

diff --git a/Engine/Source/Developer/ShaderCompilerCommon/Private/HlslAST.cpp b/Engine/Source/Developer/ShaderCompilerCommon/Private/HlslAST.cpp
new file mode 100644
--- /dev/null
+++ b/Engine/Source/Developer/ShaderCompilerCommon/Private/HlslAST.cpp
@@ -0,0 +1,44 @@
+// Copyright 1998-2014 Epic Games, Inc. All Rights Reserved.
+
+/*=============================================================================
+	HlslAST.cpp - Abstract Syntax Tree implementation for HLSL.
+=============================================================================*/
+
+#include "ShaderCompilerCommon.h"
+#include "HlslLexer.h"
+#include "HlslAST.h"
+
+namespace CrossCompiler
+{
+	namespace AST
+	{
+		FNode::FNode()
+		{
+		}
+
+		FNode::FNode(const FSourceInfo& InInfo) :
+			SourceInfo(InInfo)
+		{
+		}
+
+		FExpression::FExpression(EOperators InOperator, FExpression* E0, FExpression* E1, FExpression* E2, const FSourceInfo& InInfo) :
+			FNode(InInfo),
+			Operator(InOperator)
+		{
+			SubExpressions[0] = E0;
+			SubExpressions[1] = E1;
+			SubExpressions[2] = E2;
+			UintConstant = 0;
+		}
+
+		FUnaryExpression::FUnaryExpression(EOperators InOperator, FExpression* Expr, const FSourceInfo& InInfo) :
+			FExpression(InOperator, Expr, nullptr, nullptr, InInfo)
+		{
+		}
+
+		FBinaryExpression::FBinaryExpression(EOperators InOperator, FExpression* E0, FExpression* E1, const FSourceInfo& InInfo) :
+			FExpression(InOperator, E0, E1, nullptr, InInfo)
+		{
+		}
+	}
+}
diff --git a/Engine/Source/Developer/ShaderCompilerCommon/Private/HlslAST.h b/Engine/Source/Developer/ShaderCompilerCommon/Private/HlslAST.h
new file mode 100644
--- /dev/null
+++ b/Engine/Source/Developer/ShaderCompilerCommon/Private/HlslAST.h
@@ -0,0 +1,410 @@
+// Copyright 1998-2014 Epic Games, Inc. All Rights Reserved.
+
+/*=============================================================================
+	HlslAST.h - Abstract Syntax Tree interfaces for HLSL.
+=============================================================================*/
+
+#pragma once
+
+#include "HlslLexer.h"
+
+namespace CrossCompiler
+{
+	namespace AST
+	{
+		class FNode
+		{
+		public:
+			FSourceInfo SourceInfo;
+
+		protected:
+			FNode();
+			FNode(const FSourceInfo& InInfo);
+		};
+
+		/**
+		* Operators for AST expression nodes.
+		*/
+		enum class EOperators
+		{
+			Assign,
+			Plus,        /**< Unary + operator. */
+			Neg,
+			Add,
+			Sub,
+			Mul,
+			Div,
+			Mod,
+			LShift,
+			RShift,
+			Less,
+			Greater,
+			LEqual,
+			GEqual,
+			Equal,
+			NEqual,
+			BitAnd,
+			BitXor,
+			BitOr,
+			BitNot,
+			LogicAnd,
+			LogicXor,
+			LogicOr,
+			LogicNot,
+
+			MulAssign,
+			DivAssign,
+			ModAssign,
+			AddAssign,
+			SubAssign,
+			LSAssign,
+			RSAssign,
+			AndAssign,
+			XorAssign,
+			OrAssign,
+
+			Conditional,
+
+			PreInc,
+			PreDec,
+			//Post_inc,
+			//Post_dec,
+			//Field_selection,
+			//Array_index,
+
+			//Function_call,
+			//Initializer_list,
+
+			Identifier,
+			//Int_constant,
+			UintConstant,
+			FloatConstant,
+			BoolConstant,
+
+			//Sequence,
+
+			//Type_cast,
+		};
+
+		inline EOperators TokenToASTOperator(EHlslToken Token)
+		{
+			switch (Token)
+			{
+			case EHlslToken::Equal:
+				return AST::EOperators::Assign;
+
+			case EHlslToken::PlusEqual:
+				return AST::EOperators::AddAssign;
+
+			case EHlslToken::MinusEqual:
+				return AST::EOperators::SubAssign;
+
+			case EHlslToken::TimesEqual:
+				return AST::EOperators::MulAssign;
+
+			case EHlslToken::DivEqual:
+				return AST::EOperators::DivAssign;
+
+			case EHlslToken::ModEqual:
+				return AST::EOperators::ModAssign;
+
+			case EHlslToken::GreaterGreaterEqual:
+				return AST::EOperators::RSAssign;
+
+			case EHlslToken::LowerLowerEqual:
+				return AST::EOperators::LSAssign;
+
+			case EHlslToken::AndEqual:
+				return AST::EOperators::AndAssign;
+
+			case EHlslToken::OrEqual:
+				return AST::EOperators::OrAssign;
+
+			case EHlslToken::XorEqual:
+				return AST::EOperators::XorAssign;
+
+			case EHlslToken::Question:
+				return AST::EOperators::Conditional;
+
+			case EHlslToken::OrOr:
+				return AST::EOperators::LogicOr;
+
+			case EHlslToken::AndAnd:
+				return AST::EOperators::LogicAnd;
+
+			case EHlslToken::Or:
+				return AST::EOperators::BitOr;
+
+			case EHlslToken::Xor:
+				return AST::EOperators::BitXor;
+
+			case EHlslToken::And:
+				return AST::EOperators::BitAnd;
+
+			case EHlslToken::EqualEqual:
+				return AST::EOperators::Equal;
+
+			case EHlslToken::NotEqual:
+				return AST::EOperators::NEqual;
+
+			case EHlslToken::Lower:
+				return AST::EOperators::Less;
+
+			case EHlslToken::Greater:
+				return AST::EOperators::Greater;
+
+			case EHlslToken::LowerEqual:
+				return AST::EOperators::LEqual;
+
+			case EHlslToken::GreaterEqual:
+				return AST::EOperators::GEqual;
+
+			case EHlslToken::LowerLower:
+				return AST::EOperators::LShift;
+
+			case EHlslToken::GreaterGreater:
+				return AST::EOperators::RShift;
+
+			case EHlslToken::Plus:
+				return AST::EOperators::Add;
+
+			case EHlslToken::Minus:
+				return AST::EOperators::Sub;
+
+			case EHlslToken::Times:
+				return AST::EOperators::Mul;
+
+			case EHlslToken::Div:
+				return AST::EOperators::Div;
+
+			case EHlslToken::Mod:
+				return AST::EOperators::Mod;
+
+			default:
+				check(0);
+				break;
+			}
+
+			return AST::EOperators::Plus;
+		}
+
+		struct FExpression : public FNode
+		{
+			FExpression(EOperators InOperator, FExpression* E0, FExpression* E1, FExpression* E2, const FSourceInfo& InInfo);
+
+			EOperators Operator;
+			FExpression* SubExpressions[3];
+
+			union
+			{
+				uint32 UintConstant;
+				float FloatConstant;
+				bool BoolConstant;
+			};
+			FString Identifier;
+		};
+
+		struct FUnaryExpression : public FExpression
+		{
+			FUnaryExpression(EOperators InOperator, FExpression* Expr, const FSourceInfo& InInfo);
+		};
+
+		struct FBinaryExpression : public FExpression
+		{
+			FBinaryExpression(EOperators InOperator, FExpression* E0, FExpression* E1, const FSourceInfo& InInfo);
+		};
+/*
+		class FFunctionExpression : public FExpression
+		{
+		public:
+		private:
+		};
+
+		class FInitializerListExpression : public FExpression
+		{
+		public:
+		private:
+		};
+
+		class FCompoundStatement : public FNode
+		{
+		public:
+		private:
+		};
+
+		class FDeclaration : public FNode
+		{
+		public:
+		private:
+		};
+
+		class FStructSpecifier : public FNode
+		{
+		public:
+		private:
+		};
+
+		class FCBufferDeclaration : public FNode
+		{
+		public:
+		private:
+		};
+
+		class FTypeSpecifier : public FNode
+		{
+		public:
+		private:
+		};
+
+		class FFullySpecifiedType : public FNode
+		{
+		public:
+		private:
+		};
+
+		class FDeclaratorList : public FNode
+		{
+		public:
+		private:
+		};
+
+		class FParameterDeclarator : public FNode
+		{
+		public:
+		private:
+		};
+
+		class FFunction : public FNode
+		{
+		public:
+		private:
+		};
+
+		class FExpressionStatement : public FNode
+		{
+		public:
+		private:
+		};
+
+		class FCaseLabel : public FNode
+		{
+		public:
+		private:
+		};
+
+		class FCaseLabelList : public FNode
+		{
+		public:
+		private:
+		};
+
+		class FCaseLabelStatement : public FNode
+		{
+		public:
+		private:
+		};
+
+		class FCaseLabelStatementList : public FNode
+		{
+		public:
+		private:
+		};
+
+		class FSwitchBody : public FNode
+		{
+		public:
+		private:
+		};
+
+		class FSelectionStatement : public FNode
+		{
+		public:
+		private:
+		};
+
+		class FSwitchStatement : public FNode
+		{
+		public:
+		private:
+		};
+
+		class FIterationStatement : public FNode
+		{
+		public:
+		private:
+		};
+
+		class FJumpStatement : public FNode
+		{
+		public:
+		private:
+		};
+
+		class FFunctionDefinition : public FNode
+		{
+		public:
+		private:
+		};
+
+		class FAttributeArgument : public FNode
+		{
+		public:
+		private:
+		};
+
+		class FAttribute : public FNode
+		{
+		public:
+		private:
+		};
+*/
+	/*
+		union UASTUnion
+		{
+			int n;
+			float real;
+			const char *identifier;
+			const char *string_literal;
+
+			struct ast_type_qualifier type_qualifier;
+
+			ast_node *node;
+			ast_type_specifier *type_specifier;
+			ast_fully_specified_type *fully_specified_type;
+			ast_function *function;
+			ast_parameter_declarator *parameter_declarator;
+			ast_function_definition *function_definition;
+			ast_compound_statement *compound_statement;
+			*/
+
+/*
+			FExpression* Expression;
+*/
+/*
+			ast_declarator_list *declarator_list;
+			ast_struct_specifier *struct_specifier;
+			ast_declaration *declaration;
+			ast_switch_body *switch_body;
+			ast_case_label *case_label;
+			ast_case_label_list *case_label_list;
+			ast_case_statement *case_statement;
+			ast_case_statement_list *case_statement_list;
+
+			struct
+			{
+				ast_node *cond;
+				ast_expression *rest;
+			} for_rest_statement;
+
+			struct
+			{
+				ast_node *then_statement;
+				ast_node *else_statement;
+			} selection_rest_statement;
+
+			ast_attribute* attribute;
+			ast_attribute_arg* attribute_arg;
+		};
+	*/
+	}
+}
diff --git a/Engine/Source/Developer/ShaderCompilerCommon/Private/HlslExpressionParser.inl b/Engine/Source/Developer/ShaderCompilerCommon/Private/HlslExpressionParser.inl
--- a/Engine/Source/Developer/ShaderCompilerCommon/Private/HlslExpressionParser.inl
+++ b/Engine/Source/Developer/ShaderCompilerCommon/Private/HlslExpressionParser.inl
@@ -7,13 +7,14 @@
 #pragma once
 #include "ShaderCompilerCommon.h"
 #include "HlslParser.h"
+#include "HlslAST.h"
 
 namespace CrossCompiler
 {
 	struct FSymbolScope;
 	struct FInfo;
 
-	EParseResult ComputeExpr(FHlslScanner& Scanner, int32 MinPrec, FInfo& Info, FSymbolScope* SymbolScope);
+	EParseResult ComputeExpr(FHlslScanner& Scanner, int32 MinPrec, FInfo& Info, FSymbolScope* SymbolScope, AST::FExpression** OutExpression);
 	EParseResult ParseExpressionList(EHlslToken EndListToken, FHlslScanner& Scanner, FSymbolScope* SymbolScope, EHlslToken NewStartListToken = EHlslToken::Invalid);
 
 	struct FSymbolScope
@@ -324,43 +325,53 @@ namespace CrossCompiler
 		return EParseResult::NotMatched;
 	}
 
-	EParseResult MatchUnaryOperator(FHlslScanner& Scanner, FInfo& Info, FSymbolScope* SymbolScope)
+	// Unary!(Unary-(Unary+())) would have ! as Top, and + as Inner
+	EParseResult MatchUnaryOperator(FHlslScanner& Scanner, FInfo& Info, FSymbolScope* SymbolScope, AST::FExpression** OuterExpression, AST::FExpression** InnerExpression)
 	{
 		bool bFoundAny = false;
 		bool bTryAgain = true;
+		AST::FExpression*& PrevExpression = *InnerExpression;
 		while (Scanner.HasMoreTokens() && bTryAgain)
 		{
 			auto* Token = Scanner.GetCurrentToken();
+			AST::EOperators Operator = AST::EOperators::Plus;
+
 			switch (Token->Token)
 			{
 			case EHlslToken::PlusPlus:
 				bFoundAny = true;
 				Scanner.Advance();
+				Operator = AST::EOperators::PreInc;
 				break;
 
 			case EHlslToken::MinusMinus:
 				bFoundAny = true;
 				Scanner.Advance();
+				Operator = AST::EOperators::PreDec;
 				break;
 
 			case EHlslToken::Plus:
-				bFoundAny = true;
 				Scanner.Advance();
+				bFoundAny = true;
+				Operator = AST::EOperators::Plus;
 				break;
 
 			case EHlslToken::Minus:
-				bFoundAny = true;
 				Scanner.Advance();
+				bFoundAny = true;
+				Operator = AST::EOperators::Neg;
 				break;
 
 			case EHlslToken::Not:
-				bFoundAny = true;
 				Scanner.Advance();
+				bFoundAny = true;
+				Operator = AST::EOperators::LogicNot;
 				break;
 
 			case EHlslToken::Neg:
-				bFoundAny = true;
 				Scanner.Advance();
+				bFoundAny = true;
+				Operator = AST::EOperators::BitNot;
 				break;
 
 			case EHlslToken::LeftParenthesis:
@@ -371,6 +382,7 @@ namespace CrossCompiler
 				if (Peek1 && ParseGeneralType(Peek1, ETF_BUILTIN_NUMERIC | ETF_USER_TYPES, SymbolScope) == EParseResult::Matched && Peek2 && Peek2->Token == EHlslToken::RightParenthesis)
 				{
 					// Cast
+check(0);
 					Scanner.Advance();
 					Scanner.Advance();
 					Scanner.Advance();
@@ -387,37 +399,60 @@ namespace CrossCompiler
 			default:
 				return bFoundAny ? EParseResult::Matched : EParseResult::NotMatched;
 			}
+
+			auto* Expression = new AST::FUnaryExpression(Operator, nullptr, Token->SourceInfo);
+			if (PrevExpression)
+			{
+				PrevExpression->SubExpressions[0] = Expression;
+			}
+
+			if (!*OuterExpression)
+			{
+				*OuterExpression = Expression;
+			}
+
+			PrevExpression = Expression;
 		}
 
 		// Ran out of tokens!
 		return EParseResult::Error;
 	}
 
-	EParseResult ComputeAtom(FHlslScanner& Scanner, FInfo& Info, FSymbolScope* SymbolScope)
+	EParseResult ComputeAtom(FHlslScanner& Scanner, FInfo& Info, FSymbolScope* SymbolScope, AST::FExpression** OutExpression)
 	{
-		auto UnaryResult = MatchUnaryOperator(Scanner, Info, SymbolScope);
+		AST::FExpression* InnerUnaryExpression = nullptr;
+		auto UnaryResult = MatchUnaryOperator(Scanner, Info, SymbolScope, OutExpression, &InnerUnaryExpression);
 		auto* Token = Scanner.GetCurrentToken();
 		if (!Token || UnaryResult == EParseResult::Error)
 		{
 			return EParseResult::Error;
 		}
 
+		AST::FExpression* AtomExpression = nullptr;
 		switch (Token->Token)
 		{
 		case EHlslToken::BoolConstant:
 			Scanner.Advance();
+			AtomExpression = new AST::FUnaryExpression(AST::EOperators::BoolConstant, nullptr, Token->SourceInfo);
+			AtomExpression->BoolConstant = Token->UnsignedInteger != 0;
 			break;
 
 		case EHlslToken::UnsignedIntegerConstant:
 			Scanner.Advance();
+			AtomExpression = new AST::FUnaryExpression(AST::EOperators::UintConstant, nullptr, Token->SourceInfo);
+			AtomExpression->UintConstant = Token->UnsignedInteger;
 			break;
 
 		case EHlslToken::FloatConstant:
 			Scanner.Advance();
+			AtomExpression = new AST::FUnaryExpression(AST::EOperators::FloatConstant, nullptr, Token->SourceInfo);
+			AtomExpression->FloatConstant = Token->Float;
 			break;
 
 		case EHlslToken::Identifier:
 			Scanner.Advance();
+			AtomExpression = new AST::FUnaryExpression(AST::EOperators::Identifier, nullptr, Token->SourceInfo);
+			AtomExpression->Identifier = Token->String;
 			break;
 
 		case EHlslToken::LeftParenthesis:
@@ -428,7 +463,7 @@ namespace CrossCompiler
 			const auto* Peek1 = Scanner.PeekToken(0);
 			const auto* Peek2 = Scanner.PeekToken(1);
 			// Parenthesis expression
-			if (ComputeExpr(Scanner, 1, Info, SymbolScope) != EParseResult::Matched)
+			if (ComputeExpr(Scanner, 1, Info, SymbolScope, &AtomExpression) != EParseResult::Matched)
 			{
 				Scanner.SourceError(TEXT("Expected expression!"));
 				return EParseResult::Error;
@@ -446,11 +481,13 @@ namespace CrossCompiler
 			// Grrr handle Sampler as a variable name... This is safe here since Declarations are always handled first
 			if (ParseGeneralType(Scanner, ETF_SAMPLER_TEXTURE_BUFFER) == EParseResult::Matched)
 			{
+				check(0);
 				break;
 			}
 			// Handle float3(x,y,z)
 			else if (ParseGeneralType(Scanner, ETF_BUILTIN_NUMERIC) == EParseResult::Matched)
 			{
+				check(0);
 				if (Scanner.MatchToken(EHlslToken::LeftParenthesis))
 				{
 					auto Result = ParseExpressionList(EHlslToken::RightParenthesis, Scanner, SymbolScope);
@@ -479,14 +516,16 @@ namespace CrossCompiler
 			}
 		}
 
+		check(AtomExpression);
+
 		// Suffix
 		while (Scanner.HasMoreTokens())
 		{
 			auto* Token = Scanner.GetCurrentToken();
-			//Info.PrintWithTabs(FString::Printf(TEXT("Atom %s\n"), *Token->String));
 			if (Scanner.MatchToken(EHlslToken::LeftSquareBracket))
 			{
-				//FPlatformMisc::LowLevelOutputDebugStringf(TEXT("Array\n"));
+check(0);
+/*
 				auto Result = ComputeExpr(Scanner, 1, Info, SymbolScope);
 				if (Result != EParseResult::Matched)
 				{
@@ -499,10 +538,11 @@ namespace CrossCompiler
 					Scanner.SourceError(TEXT("Expected ']'!"));
 					return EParseResult::Error;
 				}
+*/
 			}
 			else if (Scanner.MatchToken(EHlslToken::Dot))
 			{
-				//FPlatformMisc::LowLevelOutputDebugStringf(TEXT("Member or swizzle\n"));
+				check(0);
 				if (!Scanner.MatchToken(EHlslToken::Identifier))
 				{
 					Scanner.SourceError(TEXT("Expected identifier for member or swizzle!"));
@@ -511,6 +551,7 @@ namespace CrossCompiler
 			}
 			else if (Scanner.MatchToken(EHlslToken::LeftParenthesis))
 			{
+				check(0);
 				// Function Call
 				auto Result = ParseExpressionList(EHlslToken::RightParenthesis, Scanner, SymbolScope);
 				if (Result != EParseResult::Matched)
@@ -521,11 +562,12 @@ namespace CrossCompiler
 			}
 			else if (Scanner.MatchToken(EHlslToken::PlusPlus) || Scanner.MatchToken(EHlslToken::MinusMinus))
 			{
-				//FPlatformMisc::LowLevelOutputDebugStringf(TEXT("Post %s\n"), *Token->String);
+				check(0);
 			}
 			else if (Scanner.MatchToken(EHlslToken::Question))
 			{
-				//FPlatformMisc::LowLevelOutputDebugStringf(TEXT("Ternary\n\t"));
+				check(0);
+/*
 				if (ComputeExpr(Scanner, 0, Info, SymbolScope) != EParseResult::Matched)
 				{
 					Scanner.SourceError(TEXT("Expected expression!"));
@@ -543,6 +585,7 @@ namespace CrossCompiler
 					return EParseResult::Error;
 				}
 				break;
+*/
 			}
 			else
 			{
@@ -550,6 +593,18 @@ namespace CrossCompiler
 			}
 		}
 
+		// Patch unary if necessary
+		if (InnerUnaryExpression)
+		{
+			check(!InnerUnaryExpression->SubExpressions[0]);
+			InnerUnaryExpression->SubExpressions[0] = AtomExpression;
+		}
+
+		if (!*OutExpression)
+		{
+			*OutExpression = AtomExpression;
+		}
+
 		return EParseResult::Matched;
 	}
 
@@ -626,7 +681,7 @@ namespace CrossCompiler
 		return GetPrecedence(Token) > 0;
 	}
 
-	EParseResult ComputeExpr(FHlslScanner& Scanner, int32 MinPrec, FInfo& Info, FSymbolScope* SymbolScope)
+	EParseResult ComputeExpr(FHlslScanner& Scanner, int32 MinPrec, FInfo& Info, FSymbolScope* SymbolScope, AST::FExpression** OutExpression)
 	{
 		auto OriginalToken = Scanner.GetCurrentTokenIndex();
 		FInfoIndentScope Scope(Info);
@@ -648,12 +703,12 @@ namespace CrossCompiler
 			  return result
 		*/
 		Info.PrintWithTabs(FString::Printf(TEXT("Compute Expr %d\n"), MinPrec));
-		auto Result = ComputeAtom(Scanner, Info, SymbolScope);
+		auto Result = ComputeAtom(Scanner, Info, SymbolScope, OutExpression);
 		if (Result != EParseResult::Matched)
 		{
 			return Result;
 		}
-
+		check(*OutExpression);
 		do
 		{
 			auto* Token = Scanner.GetCurrentToken();
@@ -665,7 +720,8 @@ namespace CrossCompiler
 
 			Scanner.Advance();
 			auto NextMinPrec = Precedence + 1;
-			Result = ComputeExpr(Scanner, NextMinPrec, Info, SymbolScope);
+			AST::FExpression* RHSExpression = nullptr;
+			Result = ComputeExpr(Scanner, NextMinPrec, Info, SymbolScope, &RHSExpression);
 			if (Result == EParseResult::Error)
 			{
 				return EParseResult::Error;
@@ -674,6 +730,9 @@ namespace CrossCompiler
 			{
 				break;
 			}
+			check(RHSExpression);
+			auto BinaryOperator = AST::TokenToASTOperator(Token->Token);
+			*OutExpression = new AST::FBinaryExpression(BinaryOperator, *OutExpression, RHSExpression, Token->SourceInfo);
 		}
 		while (Scanner.HasMoreTokens());
 
@@ -688,7 +747,8 @@ namespace CrossCompiler
 	EParseResult ParseExpression(FHlslScanner& Scanner, FSymbolScope* SymbolScope)
 	{
 		FInfo Info(!true);
-		return ComputeExpr(Scanner, 0, Info, SymbolScope);
+		AST::FExpression* Expression = nullptr;
+		return ComputeExpr(Scanner, 0, Info, SymbolScope, &Expression);
 	}
 
 	EParseResult ParseExpressionList(EHlslToken EndListToken, FHlslScanner& Scanner, FSymbolScope* SymbolScope, EHlslToken NewStartListToken)
