UE4 - Fix location of CCT

diff --git a/Engine/Source/Developer/ShaderCompilerCommon/Private/ShaderCompilerCommon.cpp b/Engine/Source/Developer/ShaderCompilerCommon/Private/ShaderCompilerCommon.cpp
--- a/Engine/Source/Developer/ShaderCompilerCommon/Private/ShaderCompilerCommon.cpp
+++ b/Engine/Source/Developer/ShaderCompilerCommon/Private/ShaderCompilerCommon.cpp
@@ -188,7 +188,7 @@ FString CreateCrossCompilerBatchFileContents(const FString& ShaderFile, const FS
 {
 	FString BatchFile = TEXT("@echo off");
 	BatchFile += TEXT("\nif defined ue.hlslcc GOTO DONE\nset ue.hlslcc=");
-	BatchFile += FPaths::RootDir() / TEXT("Engine\\Source\\Binaries\\Win64\\CrossCompilerTool.exe");
+	BatchFile += FPaths::RootDir() / TEXT("Engine\\Binaries\\Win64\\CrossCompilerTool.exe");
 	BatchFile += TEXT("\n\n:DONE\n%ue.hlslcc% ");
 	BatchFile += FString::Printf(TEXT("\"%s\" -o=\"%s\" %s -entry=%s %s %s"), *ShaderFile, *OutputFile, *FrequencySwitch, *EntryPoint, *VersionSwitch, *ExtraArguments);
 	BatchFile += TEXT("\npause\n");
