UE4 - HlslParser
- AST classes
- Uses linear/tape allocator with 1 MB pages
- Passed 40k shader files!

diff --git a/Engine/Source/Developer/ShaderCompilerCommon/Private/HlslAST.cpp b/Engine/Source/Developer/ShaderCompilerCommon/Private/HlslAST.cpp
--- a/Engine/Source/Developer/ShaderCompilerCommon/Private/HlslAST.cpp
+++ b/Engine/Source/Developer/ShaderCompilerCommon/Private/HlslAST.cpp
@@ -12,33 +12,1262 @@ namespace CrossCompiler
 {
 	namespace AST
 	{
-		FNode::FNode()
+		void DumpOptionalArraySize(bool bIsArray, const TLinearArray<FExpression*>& ArraySize)
 		{
+			if (bIsArray && ArraySize.Num() == 0)
+			{
+				FPlatformMisc::LowLevelOutputDebugString(TEXT("[]"));
+			}
+			else
+			{
+				for (const auto* Dimension : ArraySize)
+				{
+					FPlatformMisc::LowLevelOutputDebugString(TEXT("["));
+					Dimension->Dump(0);
+					FPlatformMisc::LowLevelOutputDebugString(TEXT("]"));
+				}
+			}
 		}
 
-		FNode::FNode(const FSourceInfo& InInfo) :
-			SourceInfo(InInfo)
+#if 0
+		FNode::FNode()/* :
+			Prev(nullptr),
+			Next(nullptr)*/
 		{
 		}
+#endif // 0
 
-		FExpression::FExpression(EOperators InOperator, FExpression* E0, FExpression* E1, FExpression* E2, const FSourceInfo& InInfo) :
-			FNode(InInfo),
-			Operator(InOperator)
+		FNode::FNode(FLinearAllocator* Allocator, const FSourceInfo& InInfo) :
+			SourceInfo(InInfo),
+			Attributes(Allocator)/*,
+			Prev(nullptr),
+			Next(nullptr)*/
+		{
+		}
+
+		void FNode::DumpIndent(int32 Indent)
+		{
+			while (--Indent >= 0)
+			{
+				FPlatformMisc::LowLevelOutputDebugString(TEXT("\t"));
+			}
+		}
+
+		void FNode::DumpAttributes() const
+		{
+			if (Attributes.Num() > 0)
+			{
+				for (auto* Attr : Attributes)
+				{
+					Attr->Dump(0);
+				}
+
+				FPlatformMisc::LowLevelOutputDebugString(TEXT(" "));
+			}
+		}
+
+		FExpression::FExpression(FLinearAllocator* InAllocator, EOperators InOperator, FExpression* E0, FExpression* E1, FExpression* E2, const FSourceInfo& InInfo) :
+			FNode(InAllocator, InInfo),
+			Operator(InOperator),
+			Identifier(nullptr),
+			Expressions(InAllocator)
 		{
 			SubExpressions[0] = E0;
 			SubExpressions[1] = E1;
 			SubExpressions[2] = E2;
-			UintConstant = 0;
+			TypeSpecifier = 0;
+		}
+
+		void FExpression::DumpOperator() const
+		{
+			switch (Operator)
+			{
+			case EOperators::Plus:
+				FPlatformMisc::LowLevelOutputDebugString(TEXT("+"));
+				break;
+
+			case EOperators::Neg:
+				FPlatformMisc::LowLevelOutputDebugString(TEXT("-"));
+				break;
+
+			case EOperators::Assign:
+				FPlatformMisc::LowLevelOutputDebugString(TEXT("="));
+				break;
+
+			case EOperators::AddAssign:
+				FPlatformMisc::LowLevelOutputDebugString(TEXT("+="));
+				break;
+
+			case EOperators::SubAssign:
+				FPlatformMisc::LowLevelOutputDebugString(TEXT("-="));
+				break;
+
+			case EOperators::MulAssign:
+				FPlatformMisc::LowLevelOutputDebugString(TEXT("*="));
+				break;
+
+			case EOperators::DivAssign:
+				FPlatformMisc::LowLevelOutputDebugString(TEXT("/="));
+				break;
+
+			case EOperators::ModAssign:
+				FPlatformMisc::LowLevelOutputDebugString(TEXT("%="));
+				break;
+
+			case EOperators::RSAssign:
+				FPlatformMisc::LowLevelOutputDebugString(TEXT(">>="));
+				break;
+
+			case EOperators::LSAssign:
+				FPlatformMisc::LowLevelOutputDebugString(TEXT("<<="));
+				break;
+
+			case EOperators::AndAssign:
+				FPlatformMisc::LowLevelOutputDebugString(TEXT("&="));
+				break;
+
+			case EOperators::OrAssign:
+				FPlatformMisc::LowLevelOutputDebugString(TEXT("|="));
+				break;
+
+			case EOperators::XorAssign:
+				FPlatformMisc::LowLevelOutputDebugString(TEXT("^="));
+				break;
+
+			case EOperators::Conditional:
+				FPlatformMisc::LowLevelOutputDebugString(TEXT("?"));
+				break;
+
+			case EOperators::LogicOr:
+				FPlatformMisc::LowLevelOutputDebugString(TEXT("||"));
+				break;
+
+			case EOperators::LogicAnd:
+				FPlatformMisc::LowLevelOutputDebugString(TEXT("&&"));
+				break;
+
+			case EOperators::LogicNot:
+				FPlatformMisc::LowLevelOutputDebugString(TEXT("!"));
+				break;
+
+			case EOperators::BitOr:
+				FPlatformMisc::LowLevelOutputDebugString(TEXT("|"));
+				break;
+
+			case EOperators::BitXor:
+				FPlatformMisc::LowLevelOutputDebugString(TEXT("^"));
+				break;
+
+			case EOperators::BitAnd:
+				FPlatformMisc::LowLevelOutputDebugString(TEXT("&"));
+				break;
+
+			case EOperators::Equal:
+				FPlatformMisc::LowLevelOutputDebugString(TEXT("=="));
+				break;
+
+			case EOperators::NEqual:
+				FPlatformMisc::LowLevelOutputDebugString(TEXT("!="));
+				break;
+
+			case EOperators::Less:
+				FPlatformMisc::LowLevelOutputDebugString(TEXT("<"));
+				break;
+
+			case EOperators::Greater:
+				FPlatformMisc::LowLevelOutputDebugString(TEXT(">"));
+				break;
+
+			case EOperators::LEqual:
+				FPlatformMisc::LowLevelOutputDebugString(TEXT("<="));
+				break;
+
+			case EOperators::GEqual:
+				FPlatformMisc::LowLevelOutputDebugString(TEXT(">="));
+				break;
+
+			case EOperators::LShift:
+				FPlatformMisc::LowLevelOutputDebugString(TEXT("<<"));
+				break;
+
+			case EOperators::RShift:
+				FPlatformMisc::LowLevelOutputDebugString(TEXT(">>"));
+				break;
+
+			case EOperators::Add:
+				FPlatformMisc::LowLevelOutputDebugString(TEXT("+"));
+				break;
+
+			case EOperators::Sub:
+				FPlatformMisc::LowLevelOutputDebugString(TEXT("-"));
+				break;
+
+			case EOperators::Mul:
+				FPlatformMisc::LowLevelOutputDebugString(TEXT("*"));
+				break;
+
+			case EOperators::Div:
+				FPlatformMisc::LowLevelOutputDebugString(TEXT("/"));
+				break;
+
+			case EOperators::Mod:
+				FPlatformMisc::LowLevelOutputDebugString(TEXT("%"));
+				break;
+
+			case EOperators::PreInc:
+				FPlatformMisc::LowLevelOutputDebugString(TEXT("++"));
+				break;
+
+			case EOperators::PreDec:
+				FPlatformMisc::LowLevelOutputDebugString(TEXT("--"));
+				break;
+
+			case EOperators::Identifier:
+				FPlatformMisc::LowLevelOutputDebugStringf(TEXT("%s"), Identifier);
+				break;
+
+			case EOperators::UintConstant:
+			case EOperators::BoolConstant:
+				FPlatformMisc::LowLevelOutputDebugStringf(TEXT("%d"), UintConstant);
+				break;
+
+			case EOperators::FloatConstant:
+				FPlatformMisc::LowLevelOutputDebugStringf(TEXT("%f"), FloatConstant);
+				break;
+
+			case EOperators::InitializerList:
+				// Nothing...
+				break;
+
+			case EOperators::PostInc:
+			case EOperators::PostDec:
+			case EOperators::FieldSelection:
+			case EOperators::ArrayIndex:
+				break;
+
+			case EOperators::TypeCast:
+				FPlatformMisc::LowLevelOutputDebugString(TEXT("("));
+				TypeSpecifier->Dump(0);
+				FPlatformMisc::LowLevelOutputDebugString(TEXT(")"));
+				break;
+
+			default:
+				FPlatformMisc::LowLevelOutputDebugStringf(TEXT("*MISSING_%d*"), Operator);
+				check(0);
+				break;
+			}
+		}
+
+		void FExpression::Dump(int32 Indent) const
+		{
+			switch (Operator)
+			{
+			case EOperators::Conditional:
+				FPlatformMisc::LowLevelOutputDebugString(TEXT("("));
+				SubExpressions[0]->Dump(0);
+				FPlatformMisc::LowLevelOutputDebugString(TEXT(" ? "));
+				SubExpressions[1]->Dump(0);
+				FPlatformMisc::LowLevelOutputDebugString(TEXT(" : "));
+				SubExpressions[2]->Dump(0);
+				FPlatformMisc::LowLevelOutputDebugString(TEXT(")"));
+				break;
+
+			default:
+				FPlatformMisc::LowLevelOutputDebugStringf(TEXT("*MISSING_%d*"), Operator);
+				check(0);
+				break;
+			}
+		}
+
+		FExpression::~FExpression()
+		{
+			for (auto* Expr : Expressions)
+			{
+				if (Expr)
+				{
+					delete Expr;
+				}
+			}
+
+			for (int32 Index = 0; Index < 3; ++Index)
+			{
+				if (SubExpressions[Index])
+				{
+					delete SubExpressions[Index];
+				}
+				else
+				{
+					break;
+				}
+			}
+		}
+
+		FUnaryExpression::FUnaryExpression(FLinearAllocator* InAllocator, EOperators InOperator, FExpression* Expr, const FSourceInfo& InInfo) :
+			FExpression(InAllocator, InOperator, Expr, nullptr, nullptr, InInfo)
+		{
+		}
+
+		void FUnaryExpression::Dump(int32 Indent) const
+		{
+			DumpOperator();
+			if (SubExpressions[0])
+			{
+				FPlatformMisc::LowLevelOutputDebugString(TEXT("("));
+				SubExpressions[0]->Dump(Indent);
+				FPlatformMisc::LowLevelOutputDebugString(TEXT(")"));
+			}
+
+			// Suffix
+			switch (Operator)
+			{
+			case EOperators::PostInc:
+				FPlatformMisc::LowLevelOutputDebugString(TEXT("++"));
+				break;
+
+			case EOperators::PostDec:
+				FPlatformMisc::LowLevelOutputDebugString(TEXT("--"));
+				break;
+
+			case EOperators::FieldSelection:
+				FPlatformMisc::LowLevelOutputDebugStringf(TEXT(".%s"), Identifier);
+				break;
+
+			default:
+				break;
+			}
+		}
+
+		FBinaryExpression::FBinaryExpression(FLinearAllocator* InAllocator, EOperators InOperator, FExpression* E0, FExpression* E1, const FSourceInfo& InInfo) :
+			FExpression(InAllocator, InOperator, E0, E1, nullptr, InInfo)
+		{
+		}
+
+		void FBinaryExpression::Dump(int32 Indent) const
+		{
+			switch (Operator)
+			{
+			case EOperators::ArrayIndex:
+				FPlatformMisc::LowLevelOutputDebugString(TEXT("("));
+				SubExpressions[0]->Dump(Indent);
+				FPlatformMisc::LowLevelOutputDebugString(TEXT(")"));
+				FPlatformMisc::LowLevelOutputDebugString(TEXT("["));
+				SubExpressions[1]->Dump(Indent);
+				FPlatformMisc::LowLevelOutputDebugString(TEXT("]"));
+				break;
+
+			default:
+				FPlatformMisc::LowLevelOutputDebugString(TEXT("("));
+				SubExpressions[0]->Dump(Indent);
+				DumpOperator();
+				SubExpressions[1]->Dump(Indent);
+				FPlatformMisc::LowLevelOutputDebugString(TEXT(")"));
+				break;
+			}
+		}
+
+		FExpressionStatement::FExpressionStatement(FLinearAllocator* InAllocator, FExpression* InExpr, const FSourceInfo& InInfo) :
+			FNode(InAllocator, InInfo),
+			Expression(InExpr)
+		{
+		}
+
+		void FExpressionStatement::Dump(int32 Indent) const
+		{
+			DumpIndent(Indent);
+			Expression->Dump(Indent);
+			FPlatformMisc::LowLevelOutputDebugString(TEXT(";\n"));
+		}
+
+		FExpressionStatement::~FExpressionStatement()
+		{
+			if (Expression)
+			{
+				delete Expression;
+			}
+		}
+
+		FCompoundStatement::FCompoundStatement(FLinearAllocator* InAllocator, const FSourceInfo& InInfo) :
+			FNode(InAllocator, InInfo),
+			Statements(InAllocator)
+		{
+		}
+
+		void FCompoundStatement::Dump(int32 Indent) const
+		{
+			DumpIndent(Indent);
+			FPlatformMisc::LowLevelOutputDebugString(TEXT("{\n"));
+			for (auto* Statement : Statements)
+			{
+				Statement->Dump(Indent + 1);
+			}
+			DumpIndent(Indent);
+			FPlatformMisc::LowLevelOutputDebugString(TEXT("}\n"));
+		}
+
+		FCompoundStatement::~FCompoundStatement()
+		{
+			for (auto* Statement : Statements)
+			{
+				if (Statement)
+				{
+					delete Statement;
+				}
+			}
+		}
+
+		FFunctionDefinition::FFunctionDefinition(FLinearAllocator* InAllocator, const FSourceInfo& InInfo) :
+			FNode(InAllocator, InInfo),
+			Prototype(nullptr),
+			Body(nullptr)
+		{
+		}
+
+		void FFunctionDefinition::Dump(int32 Indent) const
+		{
+			DumpAttributes();
+			Prototype->Dump(0);
+			if (Body)
+			{
+				Body->Dump(Indent);
+			}
+		}
+
+		FFunctionDefinition::~FFunctionDefinition()
+		{
+			delete Prototype;
+			delete Body;
+		}
+
+		FFunction::FFunction(FLinearAllocator* InAllocator, const FSourceInfo& InInfo) :
+			FNode(InAllocator, InInfo),
+			ReturnType(nullptr),
+			Identifier(nullptr),
+			ReturnSemantic(nullptr),
+			Parameters(InAllocator)
+		{
+		}
+
+		void FFunction::Dump(int32 Indent) const
+		{
+			DumpAttributes();
+			FPlatformMisc::LowLevelOutputDebugString(TEXT("\n"));
+			ReturnType->Dump(0);
+			FPlatformMisc::LowLevelOutputDebugStringf(TEXT(" %s("), Identifier);
+			bool bFirst = true;
+			for (auto* Param : Parameters)
+			{
+				if (bFirst)
+				{
+					bFirst = false;
+				}
+				else
+				{
+					FPlatformMisc::LowLevelOutputDebugString(TEXT(", "));
+				}
+				Param->Dump(0);
+			}
+
+			if (ReturnSemantic && *ReturnSemantic)
+			{
+				FPlatformMisc::LowLevelOutputDebugStringf(TEXT(": %s"), ReturnSemantic);
+			}
+
+			FPlatformMisc::LowLevelOutputDebugString(TEXT(")\n"));
+		}
+
+		FFunction::~FFunction()
+		{
+			for (auto* Param : Parameters)
+			{
+				delete Param;
+			}
+		}
+
+		FJumpStatement::FJumpStatement(FLinearAllocator* InAllocator, EJumpType InType, const FSourceInfo& InInfo) :
+			FNode(InAllocator, InInfo),
+			Type(InType),
+			OptionalExpression(nullptr)
+		{
+		}
+
+		void FJumpStatement::Dump(int32 Indent) const
+		{
+			DumpIndent(Indent);
+
+			switch (Type)
+			{
+			case EJumpType::Return:
+				FPlatformMisc::LowLevelOutputDebugString(TEXT("return"));
+				break;
+
+			case EJumpType::Break:
+				FPlatformMisc::LowLevelOutputDebugString(TEXT("break"));
+				break;
+
+			case EJumpType::Continue:
+				FPlatformMisc::LowLevelOutputDebugString(TEXT("continue"));
+				break;
+
+			default:
+				FPlatformMisc::LowLevelOutputDebugStringf(TEXT("*MISSING_%d*"), Type);
+				check(0);
+				break;
+			}
+
+			if (OptionalExpression)
+			{
+				FPlatformMisc::LowLevelOutputDebugString(TEXT(" "));
+				OptionalExpression->Dump(Indent);
+			}
+			FPlatformMisc::LowLevelOutputDebugString(TEXT(";\n"));
+		}
+
+		FJumpStatement::~FJumpStatement()
+		{
+			if (OptionalExpression)
+			{
+				delete OptionalExpression;
+			}
+		}
+
+		FSelectionStatement::FSelectionStatement(FLinearAllocator* InAllocator, const FSourceInfo& InInfo) :
+			FNode(InAllocator, InInfo),
+			Condition(nullptr),
+			ThenStatement(nullptr),
+			ElseStatement(nullptr)
+		{
+		}
+
+		void FSelectionStatement::Dump(int32 Indent) const
+		{
+			DumpIndent(Indent);
+			DumpAttributes();
+			FPlatformMisc::LowLevelOutputDebugString(TEXT("if("));
+			Condition->Dump(Indent);
+			FPlatformMisc::LowLevelOutputDebugString(TEXT(")\n"));
+			ThenStatement->Dump(Indent);
+			if (ElseStatement)
+			{
+				DumpIndent(Indent);
+				FPlatformMisc::LowLevelOutputDebugString(TEXT("else\n"));
+				ElseStatement->Dump(Indent);
+			}
+		}
+
+		FSelectionStatement::~FSelectionStatement()
+		{
+			delete Condition;
+			delete ThenStatement;
+			if (ElseStatement)
+			{
+				delete ElseStatement;
+			}
+		}
+
+		FTypeSpecifier::FTypeSpecifier(FLinearAllocator* InAllocator, const FSourceInfo& InInfo) :
+			FNode(InAllocator, InInfo),
+			Structure(nullptr),
+			TextureMSNumSamples(1),
+			PatchSize(0),
+			bIsArray(false),
+			//bIsUnsizedArray(false),
+			ArraySize(nullptr)
+		{
+		}
+
+		void FTypeSpecifier::Dump(int32 Indent) const
+		{
+			if (Structure)
+			{
+				Structure->Dump(Indent);
+			}
+			else
+			{
+				FPlatformMisc::LowLevelOutputDebugString(TypeName);
+				if (TextureMSNumSamples > 1)
+				{
+					FPlatformMisc::LowLevelOutputDebugStringf(TEXT("<%s, %d>"), InnerType, TextureMSNumSamples);
+				}
+				else if (InnerType && *InnerType)
+				{
+					FPlatformMisc::LowLevelOutputDebugStringf(TEXT("<%s>"), InnerType);
+				}
+			}
+
+			if (bIsArray)
+			{
+				printf("[ ");
+
+				if (ArraySize)
+				{
+					ArraySize->Dump(Indent);
+				}
+
+				printf("]");
+			}
+		}
+
+		FTypeSpecifier::~FTypeSpecifier()
+		{
+			if (Structure)
+			{
+				delete Structure;
+			}
+
+			if (ArraySize)
+			{
+				delete ArraySize;
+			}
+		}
+
+		FCBufferDeclaration::FCBufferDeclaration(FLinearAllocator* InAllocator, const FSourceInfo& InInfo) :
+			FNode(InAllocator, InInfo),
+			Name(nullptr),
+			Declarations(InAllocator)
+		{
+		}
+
+		void FCBufferDeclaration::Dump(int32 Indent) const
+		{
+			DumpIndent(Indent);
+			FPlatformMisc::LowLevelOutputDebugStringf(TEXT("cbuffer %s\n"), Name);
+
+			DumpIndent(Indent);
+			FPlatformMisc::LowLevelOutputDebugString(TEXT("{\n"));
+
+			for (auto* Declaration : Declarations)
+			{
+				Declaration->Dump(Indent + 1);
+			}
+
+			DumpIndent(Indent);
+			FPlatformMisc::LowLevelOutputDebugString(TEXT("}\n"));
+		}
+
+		FCBufferDeclaration::~FCBufferDeclaration()
+		{
+			for (auto* Decl : Declarations)
+			{
+				delete Decl;
+			}
+		}
+
+		FTypeQualifier::FTypeQualifier()
+		{
+			Raw = 0;
+		}
+
+		void FTypeQualifier::Dump() const
+		{
+			if (bConstant)
+			{
+				FPlatformMisc::LowLevelOutputDebugString(TEXT("const "));
+			}
+
+			if (bIsStatic)
+			{
+				FPlatformMisc::LowLevelOutputDebugString(TEXT("static "));
+			}
+
+			if (bIn && bOut)
+			{
+				FPlatformMisc::LowLevelOutputDebugString(TEXT("inout "));
+			}
+			else if (bIn)
+			{
+				FPlatformMisc::LowLevelOutputDebugString(TEXT("in "));
+			}
+			else if (bOut)
+			{
+				FPlatformMisc::LowLevelOutputDebugString(TEXT("out "));
+			}
+
+			//bRowMajor?
+			//bShared?
+		}
+
+		FFullySpecifiedType::FFullySpecifiedType(FLinearAllocator* InAllocator, const FSourceInfo& InInfo) :
+			FNode(InAllocator, InInfo),
+			Specifier(nullptr)
+		{
+		}
+
+		void FFullySpecifiedType::Dump(int32 Indent) const
+		{
+			Qualifier.Dump();
+			Specifier->Dump(Indent);
+		}
+
+		FFullySpecifiedType::~FFullySpecifiedType()
+		{
+			delete Specifier;
+		}
+
+		FDeclaration::FDeclaration(FLinearAllocator* InAllocator, const FSourceInfo& InInfo) :
+			FNode(InAllocator, InInfo),
+			Identifier(nullptr),
+			Semantic(nullptr),
+			bIsArray(false),
+			ArraySize(InAllocator),
+			Initializer(nullptr)
+		{
+		}
+
+		void FDeclaration::Dump(int32 Indent) const
+		{
+			DumpAttributes();
+			FPlatformMisc::LowLevelOutputDebugStringf(TEXT("%s"), Identifier);
+
+			DumpOptionalArraySize(bIsArray, ArraySize);
+
+			if (Initializer)
+			{
+				FPlatformMisc::LowLevelOutputDebugString(TEXT(" = "));
+				Initializer->Dump(Indent);
+			}
+
+			if (Semantic && *Semantic)
+			{
+				FPlatformMisc::LowLevelOutputDebugStringf(TEXT(" : %s"), Semantic);
+			}
+		}
+
+		FDeclaration::~FDeclaration()
+		{
+			for (auto* Expr : ArraySize)
+			{
+				delete Expr;
+			}
+
+			if (Initializer)
+			{
+				delete Initializer;
+			}
+		}
+
+		FDeclaratorList::FDeclaratorList(FLinearAllocator* InAllocator, const FSourceInfo& InInfo) :
+			FNode(InAllocator, InInfo),
+			Type(nullptr),
+			Declarations(InAllocator)
+		{
+		}
+
+		void FDeclaratorList::Dump(int32 Indent) const
+		{
+			DumpIndent(Indent);
+			DumpAttributes();
+			if (Type)
+			{
+				Type->Dump(0);
+			}
+
+			bool bFirst = true;
+			for (auto* Decl : Declarations)
+			{
+				if (bFirst)
+				{
+					FPlatformMisc::LowLevelOutputDebugString(TEXT(" "));
+					bFirst = false;
+				}
+				else
+				{
+					FPlatformMisc::LowLevelOutputDebugString(TEXT(", "));
+				}
+
+				Decl->Dump(0);
+			}
+
+			FPlatformMisc::LowLevelOutputDebugString(TEXT(";\n"));
+		}
+
+		FDeclaratorList::~FDeclaratorList()
+		{
+			delete Type;
+
+			for (auto* Decl : Declarations)
+			{
+				delete Decl;
+			}
+		}
+
+		FInitializerListExpression::FInitializerListExpression(FLinearAllocator* InAllocator, const FSourceInfo& InInfo) :
+			FExpression(InAllocator, EOperators::InitializerList, nullptr, nullptr, nullptr, InInfo)
+		{
+		}
+
+		void FInitializerListExpression::Dump(int32 Indent) const
+		{
+			FPlatformMisc::LowLevelOutputDebugString(TEXT("{"));
+			bool bFirst = true;
+			for (auto* Expr : Expressions)
+			{
+				if (bFirst)
+				{
+					bFirst = false;
+				}
+				else
+				{
+					FPlatformMisc::LowLevelOutputDebugString(TEXT(", "));
+				}
+
+				Expr->Dump(0);
+			}
+			FPlatformMisc::LowLevelOutputDebugString(TEXT("}"));
+		}
+
+		FParameterDeclarator::FParameterDeclarator(FLinearAllocator* InAllocator, const FSourceInfo& InInfo) :
+			FNode(InAllocator, InInfo),
+			Type(nullptr),
+			bIsArray(false),
+			ArraySize(InAllocator),
+			DefaultValue(nullptr)
+		{
+		}
+
+		void FParameterDeclarator::Dump(int32 Indent) const
+		{
+			DumpAttributes();
+			Type->Dump(Indent);
+			FPlatformMisc::LowLevelOutputDebugStringf(TEXT(" %s"), Identifier);
+
+			DumpOptionalArraySize(bIsArray, ArraySize);
+
+			if (Semantic && *Semantic)
+			{
+				FPlatformMisc::LowLevelOutputDebugStringf(TEXT(" : %s"), Semantic);
+			}
+
+			if (DefaultValue)
+			{
+				FPlatformMisc::LowLevelOutputDebugString(TEXT(" = "));
+				DefaultValue->Dump(Indent);
+			}
+		}
+
+		FParameterDeclarator* FParameterDeclarator::CreateFromDeclaratorList(FDeclaratorList* List, FLinearAllocator* Allocator)
+		{
+			check(List);
+			check(List->Declarations.Num() == 1);
+
+			auto* Source = (FDeclaration*)List->Declarations[0];
+			auto* New = new(Allocator) FParameterDeclarator(Allocator, Source->SourceInfo);
+			New->Type = List->Type;
+			New->Identifier = Source->Identifier;
+			New->Semantic = Source->Semantic;
+			New->bIsArray = Source->bIsArray;
+			New->ArraySize = Source->ArraySize;
+			New->DefaultValue = Source->Initializer;
+			return New;
+		}
+
+		FParameterDeclarator::~FParameterDeclarator()
+		{
+			delete Type;
+
+			for (auto* Expr : ArraySize)
+			{
+				delete Expr;
+			}
+
+			if (DefaultValue)
+			{
+				delete DefaultValue;
+			}
+		}
+
+		FIterationStatement::FIterationStatement(FLinearAllocator* InAllocator, const FSourceInfo& InInfo, EIterationType InType) :
+			FNode(InAllocator, InInfo),
+			Type(InType),
+			InitStatement(nullptr),
+			Condition(nullptr),
+			RestExpression(nullptr),
+			Body(nullptr)
+		{
+		}
+
+		void FIterationStatement::Dump(int32 Indent) const
+		{
+			DumpIndent(Indent);
+			DumpAttributes();
+			switch (Type)
+			{
+			case EIterationType::For:
+				FPlatformMisc::LowLevelOutputDebugString(TEXT("for ("));
+				if (InitStatement)
+				{
+					InitStatement->Dump(0);
+					DumpIndent(Indent + 1);
+				}
+				else
+				{
+					FPlatformMisc::LowLevelOutputDebugString(TEXT("; "));
+				}
+				if (Condition)
+				{
+					Condition->Dump(0);
+				}
+				FPlatformMisc::LowLevelOutputDebugString(TEXT(";"));
+				if (RestExpression)
+				{
+					RestExpression->Dump(0);
+				}
+				FPlatformMisc::LowLevelOutputDebugString(TEXT(")\n"));
+				DumpIndent(Indent);
+				FPlatformMisc::LowLevelOutputDebugString(TEXT("{\n"));
+				if (Body)
+				{
+					Body->Dump(Indent + 1);
+				}
+				DumpIndent(Indent);
+				FPlatformMisc::LowLevelOutputDebugString(TEXT("}\n"));
+				break;
+
+			case EIterationType::While:
+				FPlatformMisc::LowLevelOutputDebugString(TEXT("while ("));
+				Condition->Dump(0);
+				FPlatformMisc::LowLevelOutputDebugString(TEXT(")\n"));
+				DumpIndent(Indent);
+				FPlatformMisc::LowLevelOutputDebugString(TEXT("{\n"));
+				if (Body)
+				{
+					Body->Dump(Indent + 1);
+				}
+				DumpIndent(Indent);
+				FPlatformMisc::LowLevelOutputDebugString(TEXT("}\n"));
+				break;
+
+			case EIterationType::DoWhile:
+				FPlatformMisc::LowLevelOutputDebugString(TEXT("do\n"));
+				DumpIndent(Indent);
+				FPlatformMisc::LowLevelOutputDebugString(TEXT("{\n"));
+				if (Body)
+				{
+					Body->Dump(Indent + 1);
+				}
+				DumpIndent(Indent);
+				FPlatformMisc::LowLevelOutputDebugString(TEXT("}\n"));
+				DumpIndent(Indent);
+				FPlatformMisc::LowLevelOutputDebugString(TEXT("while ("));
+				Condition->Dump(0);
+				FPlatformMisc::LowLevelOutputDebugString(TEXT(");\n"));
+				break;
+
+			default:
+				check(0);
+				break;
+			}
+		}
+
+		FIterationStatement::~FIterationStatement()
+		{
+			if (InitStatement)
+			{
+				delete InitStatement;
+			}
+			if (Condition)
+			{
+				delete Condition;
+			}
+
+			if (RestExpression)
+			{
+				delete RestExpression;
+			}
+
+			if (Body)
+			{
+				delete Body;
+			}
+		}
+
+		FFunctionExpression::FFunctionExpression(FLinearAllocator* InAllocator, const FSourceInfo& InInfo, FExpression* Callee) :
+			FExpression(InAllocator, EOperators::FunctionCall, Callee, nullptr, nullptr, InInfo)
+		{
+		}
+
+		void FFunctionExpression::Dump(int32 Indent) const
+		{
+			SubExpressions[0]->Dump(0);
+			FPlatformMisc::LowLevelOutputDebugString(TEXT("("));
+			bool bFirst = true;
+			for (auto* Expr : Expressions)
+			{
+				if (bFirst)
+				{
+					bFirst = false;
+				}
+				else
+				{
+					FPlatformMisc::LowLevelOutputDebugString(TEXT(", "));
+				}
+				Expr->Dump(0);
+			}
+			FPlatformMisc::LowLevelOutputDebugString(TEXT(")"));
+		}
+
+		FSwitchStatement::FSwitchStatement(FLinearAllocator* InAllocator, const FSourceInfo& InInfo, FExpression* InCondition, FSwitchBody* InBody) :
+			FNode(InAllocator, InInfo),
+			Condition(InCondition),
+			Body(InBody)
+		{
+		}
+
+		void FSwitchStatement::Dump(int32 Indent) const
+		{
+			DumpIndent(Indent);
+			FPlatformMisc::LowLevelOutputDebugString(TEXT("switch ("));
+			Condition->Dump(0);
+			FPlatformMisc::LowLevelOutputDebugString(TEXT(")\n"));
+			Body->Dump(Indent);
+		}
+
+		FSwitchStatement::~FSwitchStatement()
+		{
+			delete Condition;
+			delete Body;
+		}
+
+		FSwitchBody::FSwitchBody(FLinearAllocator* InAllocator, const FSourceInfo& InInfo) :
+			FNode(InAllocator, InInfo),
+			CaseList(nullptr)
+		{
+		}
+
+		void FSwitchBody::Dump(int32 Indent) const
+		{
+			DumpIndent(Indent);
+			FPlatformMisc::LowLevelOutputDebugString(TEXT("{\n"));
+			CaseList->Dump(Indent + 1);
+			DumpIndent(Indent);
+			FPlatformMisc::LowLevelOutputDebugString(TEXT("}\n"));
+		}
+
+		FSwitchBody::~FSwitchBody()
+		{
+			delete CaseList;
+		}
+
+		FCaseLabel::FCaseLabel(FLinearAllocator* InAllocator, const FSourceInfo& InInfo, AST::FExpression* InExpression) :
+			FNode(InAllocator, InInfo),
+			TestExpression(InExpression)
+		{
+		}
+
+		void FCaseLabel::Dump(int32 Indent) const
+		{
+			DumpIndent(Indent);
+			if (TestExpression)
+			{
+				FPlatformMisc::LowLevelOutputDebugString(TEXT("case "));
+				TestExpression->Dump(0);
+			}
+			else
+			{
+				FPlatformMisc::LowLevelOutputDebugString(TEXT("default"));
+			}
+
+			FPlatformMisc::LowLevelOutputDebugString(TEXT(":\n"));
+		}
+
+		FCaseLabel::~FCaseLabel()
+		{
+			if (TestExpression)
+			{
+				delete TestExpression;
+			}
+		}
+
+
+		FCaseStatement::FCaseStatement(FLinearAllocator* InAllocator, const FSourceInfo& InInfo, FCaseLabelList* InLabels) :
+			FNode(InAllocator, InInfo),
+			Labels(InLabels),
+			Statements(InAllocator)
+		{
+		}
+
+		void FCaseStatement::Dump(int32 Indent) const
+		{
+			Labels->Dump(Indent);
+
+			if (Statements.Num() > 1)
+			{
+				DumpIndent(Indent);
+				FPlatformMisc::LowLevelOutputDebugString(TEXT("{\n"));
+				for (auto* Statement : Statements)
+				{
+					Statement->Dump(Indent + 1);
+				}
+				DumpIndent(Indent);
+				FPlatformMisc::LowLevelOutputDebugString(TEXT("}\n"));
+			}
+			else if (Statements.Num() > 0)
+			{
+				Statements[0]->Dump(Indent + 1);
+			}
+		}
+
+		FCaseStatement::~FCaseStatement()
+		{
+			delete Labels;
+			for (auto* Statement : Statements)
+			{
+				if (Statement)
+				{
+					delete Statement;
+				}
+			}
+		}
+
+		FCaseLabelList::FCaseLabelList(FLinearAllocator* InAllocator, const FSourceInfo& InInfo) :
+			FNode(InAllocator, InInfo),
+			Labels(InAllocator)
+		{
+		}
+
+		void FCaseLabelList::Dump(int32 Indent) const
+		{
+			for (auto* Label : Labels)
+			{
+				Label->Dump(Indent);
+			}
+		}
+
+		FCaseLabelList::~FCaseLabelList()
+		{
+			for (auto* Label : Labels)
+			{
+				delete Label;
+			}
+		}
+
+		FCaseStatementList::FCaseStatementList(FLinearAllocator* InAllocator, const FSourceInfo& InInfo) :
+			FNode(InAllocator, InInfo),
+			Cases(InAllocator)
+		{
+		}
+
+		void FCaseStatementList::Dump(int32 Indent) const
+		{
+			for (auto* Case : Cases)
+			{
+				Case->Dump(Indent);
+			}
+		}
+
+		FCaseStatementList::~FCaseStatementList()
+		{
+			for (auto* Case : Cases)
+			{
+				delete Case;
+			}
+		}
+
+		FStructSpecifier::FStructSpecifier(FLinearAllocator* InAllocator, const FSourceInfo& InInfo) :
+			FNode(InAllocator, InInfo),
+			Name(nullptr),
+			ParentName(nullptr),
+			Declarations(InAllocator)
+		{
+		}
+
+		void FStructSpecifier::Dump(int32 Indent) const
+		{
+			DumpIndent(Indent);
+			FPlatformMisc::LowLevelOutputDebugStringf(TEXT("struct %s"), Name);
+			if (ParentName && *ParentName)
+			{
+				FPlatformMisc::LowLevelOutputDebugStringf(TEXT(" : %s"), ParentName);
+			}
+			FPlatformMisc::LowLevelOutputDebugStringf(TEXT("\n"));
+
+			DumpIndent(Indent);
+			FPlatformMisc::LowLevelOutputDebugString(TEXT("{\n"));
+
+			for (auto* Decl : Declarations)
+			{
+				Decl->Dump(Indent + 1);
+			}
+
+			DumpIndent(Indent);
+			FPlatformMisc::LowLevelOutputDebugString(TEXT("}"));
+		}
+
+		FStructSpecifier::~FStructSpecifier()
+		{
+			for (auto* Decl : Declarations)
+			{
+				delete Decl;
+			}
+		}
+
+		FAttribute::FAttribute(FLinearAllocator* InAllocator, const FSourceInfo& InInfo, const TCHAR* InName) :
+			FNode(InAllocator, InInfo),
+			Name(InName),
+			Arguments(InAllocator)
+		{
+		}
+
+		void FAttribute::Dump(int32 Indent) const
+		{
+			FPlatformMisc::LowLevelOutputDebugStringf(TEXT("[%s"), Name);
+
+			bool bFirst = true;
+			for (auto* Arg : Arguments)
+			{
+				if (bFirst)
+				{
+					FPlatformMisc::LowLevelOutputDebugString(TEXT("("));
+					bFirst = false;
+				}
+				else
+				{
+					FPlatformMisc::LowLevelOutputDebugString(TEXT(", "));
+				}
+
+				Arg->Dump(0);
+			}
+
+			if (!bFirst)
+			{
+				FPlatformMisc::LowLevelOutputDebugString(TEXT(")"));
+			}
+
+			FPlatformMisc::LowLevelOutputDebugString(TEXT("]"));
+		}
+
+		FAttribute::~FAttribute()
+		{
+			for (auto* Arg : Arguments)
+			{
+				delete Arg;
+			}
+		}
+
+		FAttributeArgument::FAttributeArgument(FLinearAllocator* InAllocator, const FSourceInfo& InInfo) :
+			FNode(InAllocator, InInfo),
+			StringArgument(nullptr),
+			ExpressionArgument(nullptr)
+		{
 		}
 
-		FUnaryExpression::FUnaryExpression(EOperators InOperator, FExpression* Expr, const FSourceInfo& InInfo) :
-			FExpression(InOperator, Expr, nullptr, nullptr, InInfo)
+		void FAttributeArgument::Dump(int32 Indent) const
 		{
+			if (ExpressionArgument)
+			{
+				ExpressionArgument->Dump(0);
+			}
+			else
+			{
+				FPlatformMisc::LowLevelOutputDebugStringf(TEXT("\"%s\""), StringArgument);
+			}
 		}
 
-		FBinaryExpression::FBinaryExpression(EOperators InOperator, FExpression* E0, FExpression* E1, const FSourceInfo& InInfo) :
-			FExpression(InOperator, E0, E1, nullptr, InInfo)
+		FAttributeArgument::~FAttributeArgument()
 		{
+			if (ExpressionArgument)
+			{
+				delete ExpressionArgument;
+			}
 		}
 	}
 }
diff --git a/Engine/Source/Developer/ShaderCompilerCommon/Private/HlslAST.h b/Engine/Source/Developer/ShaderCompilerCommon/Private/HlslAST.h
--- a/Engine/Source/Developer/ShaderCompilerCommon/Private/HlslAST.h
+++ b/Engine/Source/Developer/ShaderCompilerCommon/Private/HlslAST.h
@@ -8,8 +8,190 @@
 
 #include "HlslLexer.h"
 
+#define USE_UNREAL_ALLOCATOR 0
+
 namespace CrossCompiler
 {
+	struct FLinearAllocator
+	{
+		static const int PageSize = 1024* 1024;
+
+		FLinearAllocator()
+		{
+			auto* Initial = new FPage(PageSize);
+			Pages.Add(Initial);
+		}
+
+		~FLinearAllocator()
+		{
+			for (auto* Page : Pages)
+			{
+				delete Page;
+			}
+		}
+
+		inline void* Alloc(SIZE_T NumBytes)
+		{
+			check(NumBytes <= PageSize);
+			auto* Page = Pages.Last();
+			if (Page->Current + NumBytes > Page->End)
+			{
+				Page = new FPage(PageSize);
+				Pages.Add(Page);
+			}
+
+			void* Ptr = Page->Current;
+			Page->Current += NumBytes;
+			return Ptr;
+		}
+
+		inline void* Alloc(SIZE_T NumBytes, SIZE_T Align)
+		{
+			void* Data = Alloc(NumBytes + Align - 1);
+			UPTRINT Address = (UPTRINT)Data;
+			Address += (Align - (Address % (UPTRINT)Align)) % Align;
+			return (void*)Address;
+		}
+
+		inline TCHAR* Strdup(const TCHAR* String)
+		{
+			if (!String)
+			{
+				return nullptr;
+			}
+
+			auto* Data = (TCHAR*)Alloc(FCString::Strlen(String) + 1);
+			FCString::Strcpy(Data, FCString::Strlen(String) + 1, String);
+			return Data;
+		}
+
+		inline TCHAR* Strdup(const FString& String)
+		{
+			auto* Data = (TCHAR*)Alloc(String.Len() + 1);
+			FCString::Strcpy(Data, String.Len() + 1, *String);
+			return Data;
+		}
+
+		struct FPage
+		{
+			char* Current;
+			char* Begin;
+			char* End;
+
+			FPage(SIZE_T Size)
+			{
+				Begin = new char[Size];
+				End = Begin + Size;
+				Current = Begin;
+			}
+
+			~FPage()
+			{
+				delete [] Begin;
+			}
+		};
+		TArray<FPage*> Pages;
+	};
+
+#if !USE_UNREAL_ALLOCATOR
+	class FLinearAllocatorPolicy
+	{
+	public:
+		// Unreal allocator magic
+		enum { NeedsElementType = false };
+		enum { RequireRangeCheck = true };
+
+		template<typename ElementType>
+		class ForElementType
+		{
+		public:
+
+			/** Default constructor. */
+			ForElementType() :
+				LinearAllocator(nullptr),
+				Data(nullptr)
+			{}
+
+			// FContainerAllocatorInterface
+			/*FORCEINLINE*/ ElementType* GetAllocation() const
+			{
+				return Data;
+			}
+			void ResizeAllocation(int32 PreviousNumElements, int32 NumElements, int32 NumBytesPerElement)
+			{
+				void* OldData = Data;
+				if (NumElements)
+				{
+					// Allocate memory from the stack.
+					Data = (ElementType*)LinearAllocator->Alloc(NumElements * NumBytesPerElement,
+						FMath::Max((uint32)sizeof(void*), (uint32)ALIGNOF(ElementType))
+						);
+
+					// If the container previously held elements, copy them into the new allocation.
+					if (OldData && PreviousNumElements)
+					{
+						const int32 NumCopiedElements = FMath::Min(NumElements, PreviousNumElements);
+						FMemory::Memcpy(Data, OldData, NumCopiedElements * NumBytesPerElement);
+					}
+				}
+			}
+			int32 CalculateSlack(int32 NumElements, int32 NumAllocatedElements, int32 NumBytesPerElement) const
+			{
+				return DefaultCalculateSlack(NumElements, NumAllocatedElements, NumBytesPerElement);
+			}
+
+			int32 GetAllocatedSize(int32 NumAllocatedElements, int32 NumBytesPerElement) const
+			{
+				return NumAllocatedElements * NumBytesPerElement;
+			}
+
+			FLinearAllocator* LinearAllocator;
+
+		private:
+
+			/** A pointer to the container's elements. */
+			ElementType* Data;
+		};
+
+		typedef ForElementType<FScriptContainerElement> ForAnyElementType;
+	};
+
+	class FLinearBitArrayAllocator
+		: public TInlineAllocator<4, FLinearAllocatorPolicy>
+	{
+	};
+
+	class FLinearSparseArrayAllocator
+		: public TSparseArrayAllocator<FLinearAllocatorPolicy, FLinearBitArrayAllocator>
+	{
+	};
+
+	class FLinearSetAllocator
+		: public TSetAllocator<FLinearSparseArrayAllocator, TInlineAllocator<1, FLinearAllocatorPolicy> >
+	{
+	};
+
+	template <typename TType>
+	class TLinearArray : public TArray<TType, FLinearAllocatorPolicy>
+	{
+	public:
+		TLinearArray(FLinearAllocator* Allocator)
+		{
+			TArray<TType, FLinearAllocatorPolicy>::AllocatorInstance.LinearAllocator = Allocator;
+		}
+	};
+
+/*
+	template <typename TType>
+	struct TLinearSet : public TSet<typename TType, DefaultKeyFuncs<typename TType>, FLinearSetAllocator>
+	{
+		TLinearSet(FLinearAllocator* InAllocator)
+		{
+			Elements.AllocatorInstance.LinearAllocator = InAllocator;
+		}
+	};*/
+#endif
+
 	namespace AST
 	{
 		class FNode
@@ -17,9 +199,42 @@ namespace CrossCompiler
 		public:
 			FSourceInfo SourceInfo;
 
+			TLinearArray<struct FAttribute*> Attributes;
+
+			virtual void Dump(int32 Indent) const = 0;
+
+			/* Callers of this ralloc-based new need not call delete. It's
+			* easier to just ralloc_free 'ctx' (or any of its ancestors). */
+			static void* operator new(SIZE_T Size, FLinearAllocator* Allocator)
+			{
+#if USE_UNREAL_ALLOCATOR
+				return FMemory::Malloc(Size);
+#else
+				auto* Ptr = Allocator->Alloc(Size);
+				return Ptr;
+#endif
+			}
+
+			/* If the user *does* call delete, that's OK, we will just
+			* ralloc_free in that case. */
+			static void operator delete(void* Pointer)
+			{
+#if USE_UNREAL_ALLOCATOR
+				FMemory::Free(Pointer);
+#else
+				// Do nothing...
+#endif
+			}
+
+			virtual ~FNode() {}
+
 		protected:
-			FNode();
-			FNode(const FSourceInfo& InInfo);
+			//FNode();
+			FNode(FLinearAllocator* Allocator, const FSourceInfo& InInfo);
+
+			static void DumpIndent(int32 Indent);
+
+			void DumpAttributes() const;
 		};
 
 		/**
@@ -67,13 +282,13 @@ namespace CrossCompiler
 
 			PreInc,
 			PreDec,
-			//Post_inc,
-			//Post_dec,
-			//Field_selection,
-			//Array_index,
+			PostInc,
+			PostDec,
+			FieldSelection,
+			ArrayIndex,
 
-			//Function_call,
-			//Initializer_list,
+			FunctionCall,
+			InitializerList,
 
 			Identifier,
 			//Int_constant,
@@ -83,7 +298,7 @@ namespace CrossCompiler
 
 			//Sequence,
 
-			//Type_cast,
+			TypeCast,
 		};
 
 		inline EOperators TokenToASTOperator(EHlslToken Token)
@@ -190,7 +405,8 @@ namespace CrossCompiler
 
 		struct FExpression : public FNode
 		{
-			FExpression(EOperators InOperator, FExpression* E0, FExpression* E1, FExpression* E2, const FSourceInfo& InInfo);
+			FExpression(FLinearAllocator* InAllocator, EOperators InOperator, FExpression* E0, FExpression* E1, FExpression* E2, const FSourceInfo& InInfo);
+			~FExpression();
 
 			EOperators Operator;
 			FExpression* SubExpressions[3];
@@ -200,211 +416,354 @@ namespace CrossCompiler
 				uint32 UintConstant;
 				float FloatConstant;
 				bool BoolConstant;
+				struct FTypeSpecifier* TypeSpecifier;
 			};
-			FString Identifier;
+			const TCHAR* Identifier;
+
+			TLinearArray<FExpression*> Expressions;
+
+			void DumpOperator() const;
+			virtual void Dump(int32 Indent) const;
 		};
 
 		struct FUnaryExpression : public FExpression
 		{
-			FUnaryExpression(EOperators InOperator, FExpression* Expr, const FSourceInfo& InInfo);
+			FUnaryExpression(FLinearAllocator* InAllocator, EOperators InOperator, FExpression* Expr, const FSourceInfo& InInfo);
+
+			virtual void Dump(int32 Indent) const;
 		};
 
 		struct FBinaryExpression : public FExpression
 		{
-			FBinaryExpression(EOperators InOperator, FExpression* E0, FExpression* E1, const FSourceInfo& InInfo);
+			FBinaryExpression(FLinearAllocator* InAllocator, EOperators InOperator, FExpression* E0, FExpression* E1, const FSourceInfo& InInfo);
+
+			virtual void Dump(int32 Indent) const;
 		};
-/*
-		class FFunctionExpression : public FExpression
+
+		struct FFunctionExpression : public FExpression
 		{
-		public:
-		private:
+			FFunctionExpression(FLinearAllocator* InAllocator, const FSourceInfo& InInfo, FExpression* Callee);
+
+			virtual void Dump(int32 Indent) const;
 		};
 
-		class FInitializerListExpression : public FExpression
+		struct FInitializerListExpression : public FExpression
 		{
-		public:
-		private:
+			FInitializerListExpression(FLinearAllocator* InAllocator, const FSourceInfo& InInfo);
+
+			virtual void Dump(int32 Indent) const;
 		};
 
-		class FCompoundStatement : public FNode
+		struct FCompoundStatement : public FNode
 		{
-		public:
-		private:
+			FCompoundStatement(FLinearAllocator* InAllocator, const FSourceInfo& InInfo);
+			~FCompoundStatement();
+
+			TLinearArray<FNode*> Statements;
+
+			virtual void Dump(int32 Indent) const;
 		};
 
-		class FDeclaration : public FNode
+		struct FDeclaration : public FNode
 		{
-		public:
-		private:
+			FDeclaration(FLinearAllocator* InAllocator, const FSourceInfo& InInfo);
+			~FDeclaration();
+
+			virtual void Dump(int32 Indent) const;
+
+			const TCHAR* Identifier;
+
+			const TCHAR* Semantic;
+
+			bool bIsArray;
+			//bool bIsUnsizedArray;
+
+			TLinearArray<FExpression*> ArraySize;
+
+			FExpression* Initializer;
 		};
 
-		class FStructSpecifier : public FNode
+		struct FTypeQualifier
 		{
-		public:
-		private:
+			union
+			{
+				struct
+				{
+					uint32 bIsStatic : 1;
+					uint32 bConstant : 1;
+					uint32 bIn : 1;
+					uint32 bOut : 1;
+					uint32 bRowMajor : 1;
+					uint32 bShared : 1;
+				};
+				uint32 Raw;
+			};
+
+			FTypeQualifier();
+
+			void Dump() const;
 		};
 
-		class FCBufferDeclaration : public FNode
+		struct FStructSpecifier : public FNode
 		{
-		public:
-		private:
+			FStructSpecifier(FLinearAllocator* InAllocator, const FSourceInfo& InInfo);
+			~FStructSpecifier();
+
+			virtual void Dump(int32 Indent) const;
+
+			TCHAR* Name;
+			TCHAR* ParentName;
+			TLinearArray<FNode*> Declarations;
 		};
 
-		class FTypeSpecifier : public FNode
+		struct FCBufferDeclaration : public FNode
 		{
-		public:
-		private:
+			FCBufferDeclaration(FLinearAllocator* InAllocator, const FSourceInfo& InInfo);
+			~FCBufferDeclaration();
+
+			virtual void Dump(int32 Indent) const;
+
+			const TCHAR* Name;
+			TLinearArray<FNode*> Declarations;
 		};
 
-		class FFullySpecifiedType : public FNode
+		struct FTypeSpecifier : public FNode
 		{
-		public:
-		private:
+			FTypeSpecifier(FLinearAllocator* InAllocator, const FSourceInfo& InInfo);
+			~FTypeSpecifier();
+
+			virtual void Dump(int32 Indent) const;
+
+			const TCHAR* TypeName;
+			const TCHAR* InnerType;
+
+			FStructSpecifier* Structure;
+
+			int32 TextureMSNumSamples;
+
+			int32 PatchSize;
+
+			bool bIsArray;
+			//bool bIsUnsizedArray;
+			FExpression* ArraySize;
 		};
 
-		class FDeclaratorList : public FNode
+		struct FFullySpecifiedType : public FNode
 		{
-		public:
-		private:
+			FFullySpecifiedType(FLinearAllocator* InAllocator, const FSourceInfo& InInfo);
+			~FFullySpecifiedType();
+
+			virtual void Dump(int32 Indent) const;
+
+			FTypeQualifier Qualifier;
+			FTypeSpecifier* Specifier;
 		};
 
-		class FParameterDeclarator : public FNode
+		struct FDeclaratorList : public FNode
 		{
-		public:
-		private:
+			FDeclaratorList(FLinearAllocator* InAllocator, const FSourceInfo& InInfo);
+			~FDeclaratorList();
+
+			virtual void Dump(int32 Indent) const;
+
+			FFullySpecifiedType* Type;
+			TLinearArray<FNode*> Declarations;
 		};
 
-		class FFunction : public FNode
+		struct FParameterDeclarator : public FNode
 		{
-		public:
-		private:
+			FParameterDeclarator(FLinearAllocator* InAllocator, const FSourceInfo& InInfo);
+			~FParameterDeclarator();
+
+			static FParameterDeclarator* CreateFromDeclaratorList(FDeclaratorList* List, FLinearAllocator* Allocator);
+
+			virtual void Dump(int32 Indent) const;
+
+			FFullySpecifiedType* Type;
+			const TCHAR* Identifier;
+			const TCHAR* Semantic;
+			bool bIsArray;
+
+			TLinearArray<FExpression*> ArraySize;
+			FExpression* DefaultValue;
 		};
 
-		class FExpressionStatement : public FNode
+		struct FFunction : public FNode
 		{
-		public:
-		private:
+			FFunction(FLinearAllocator* InAllocator, const FSourceInfo& InInfo);
+			~FFunction();
+
+			virtual void Dump(int32 Indent) const;
+
+			FFullySpecifiedType* ReturnType;
+			const TCHAR* Identifier;
+			const TCHAR* ReturnSemantic;
+
+			TLinearArray<FNode*> Parameters;
+
+			bool bIsDefinition;
+
+			//Signature
 		};
 
-		class FCaseLabel : public FNode
+		struct FExpressionStatement : public FNode
 		{
-		public:
-		private:
+			FExpressionStatement(FLinearAllocator* InAllocator, FExpression* InExpr, const FSourceInfo& InInfo);
+			~FExpressionStatement();
+
+			FExpression* Expression;
+
+			virtual void Dump(int32 Indent) const;
 		};
 
-		class FCaseLabelList : public FNode
+		struct FCaseLabel : public FNode
 		{
-		public:
-		private:
+			FCaseLabel(FLinearAllocator* InAllocator, const FSourceInfo& InInfo, AST::FExpression* InExpression);
+			~FCaseLabel();
+
+			virtual void Dump(int32 Indent) const;
+
+			FExpression* TestExpression;
 		};
 
-		class FCaseLabelStatement : public FNode
+		struct FCaseLabelList : public FNode
 		{
-		public:
-		private:
+			FCaseLabelList(FLinearAllocator* InAllocator, const FSourceInfo& InInfo);
+			~FCaseLabelList();
+
+			virtual void Dump(int32 Indent) const;
+
+			TLinearArray<FCaseLabel*> Labels;
 		};
 
-		class FCaseLabelStatementList : public FNode
+		struct FCaseStatement : public FNode
 		{
-		public:
-		private:
+			FCaseStatement(FLinearAllocator* InAllocator, const FSourceInfo& InInfo, FCaseLabelList* InLabels);
+			~FCaseStatement();
+
+			virtual void Dump(int32 Indent) const;
+
+			FCaseLabelList* Labels;
+			TLinearArray<FNode*> Statements;
 		};
 
-		class FSwitchBody : public FNode
+		struct FCaseStatementList : public FNode
 		{
-		public:
-		private:
+			FCaseStatementList(FLinearAllocator* InAllocator, const FSourceInfo& InInfo);
+			~FCaseStatementList();
+
+			virtual void Dump(int32 Indent) const;
+
+			TLinearArray<FCaseStatement*> Cases;
 		};
 
-		class FSelectionStatement : public FNode
+		struct FSwitchBody : public FNode
 		{
-		public:
-		private:
+			FSwitchBody(FLinearAllocator* InAllocator, const FSourceInfo& InInfo);
+			~FSwitchBody();
+
+			virtual void Dump(int32 Indent) const;
+
+			FCaseStatementList* CaseList;
 		};
 
-		class FSwitchStatement : public FNode
+		struct FSelectionStatement : public FNode
 		{
-		public:
-		private:
+			FSelectionStatement(FLinearAllocator* InAllocator, const FSourceInfo& InInfo);
+			~FSelectionStatement();
+
+			virtual void Dump(int32 Indent) const;
+
+			FExpression* Condition;
+
+			FNode* ThenStatement;
+			FNode* ElseStatement;
 		};
 
-		class FIterationStatement : public FNode
+		struct FSwitchStatement : public FNode
 		{
-		public:
-		private:
+			FSwitchStatement(FLinearAllocator* InAllocator, const FSourceInfo& InInfo, FExpression* InCondition, FSwitchBody* InBody);
+			~FSwitchStatement();
+
+			virtual void Dump(int32 Indent) const;
+
+			FExpression* Condition;
+			FSwitchBody* Body;
 		};
 
-		class FJumpStatement : public FNode
+		enum class EIterationType
 		{
-		public:
-		private:
+			For,
+			While,
+			DoWhile,
 		};
 
-		class FFunctionDefinition : public FNode
+		struct FIterationStatement : public FNode
 		{
-		public:
-		private:
+			FIterationStatement(FLinearAllocator* InAllocator, const FSourceInfo& InInfo, EIterationType InType);
+			~FIterationStatement();
+
+			virtual void Dump(int32 Indent) const;
+
+			EIterationType Type;
+
+			AST::FNode* InitStatement;
+			AST::FNode* Condition;
+			FExpression* RestExpression;
+
+			AST::FNode* Body;
 		};
 
-		class FAttributeArgument : public FNode
+		enum class EJumpType
 		{
-		public:
-		private:
+			Continue,
+			Break,
+			Return,
+			//Discard,
 		};
 
-		class FAttribute : public FNode
+		struct FJumpStatement : public FNode
 		{
-		public:
-		private:
+			FJumpStatement(FLinearAllocator* InAllocator, EJumpType InType, const FSourceInfo& InInfo);
+			~FJumpStatement();
+
+			EJumpType Type;
+			FExpression* OptionalExpression;
+
+			virtual void Dump(int32 Indent) const;
 		};
-*/
-	/*
-		union UASTUnion
+
+		struct FFunctionDefinition : public FNode
 		{
-			int n;
-			float real;
-			const char *identifier;
-			const char *string_literal;
+			FFunctionDefinition(FLinearAllocator* InAllocator, const FSourceInfo& InInfo);
+			~FFunctionDefinition();
 
-			struct ast_type_qualifier type_qualifier;
+			FFunction* Prototype;
+			FCompoundStatement* Body;
 
-			ast_node *node;
-			ast_type_specifier *type_specifier;
-			ast_fully_specified_type *fully_specified_type;
-			ast_function *function;
-			ast_parameter_declarator *parameter_declarator;
-			ast_function_definition *function_definition;
-			ast_compound_statement *compound_statement;
-			*/
+			virtual void Dump(int32 Indent) const;
+		};
 
-/*
-			FExpression* Expression;
-*/
-/*
-			ast_declarator_list *declarator_list;
-			ast_struct_specifier *struct_specifier;
-			ast_declaration *declaration;
-			ast_switch_body *switch_body;
-			ast_case_label *case_label;
-			ast_case_label_list *case_label_list;
-			ast_case_statement *case_statement;
-			ast_case_statement_list *case_statement_list;
-
-			struct
-			{
-				ast_node *cond;
-				ast_expression *rest;
-			} for_rest_statement;
+		struct FAttributeArgument : public FNode
+		{
+			FAttributeArgument(FLinearAllocator* InAllocator, const FSourceInfo& InInfo);
+			~FAttributeArgument();
 
-			struct
-			{
-				ast_node *then_statement;
-				ast_node *else_statement;
-			} selection_rest_statement;
+			virtual void Dump(int32 Indent) const;
+
+			const TCHAR* StringArgument;
+			FExpression* ExpressionArgument;
+		};
+
+		struct FAttribute : public FNode
+		{
+			FAttribute(FLinearAllocator* Allocator, const FSourceInfo& InInfo, const TCHAR* InName);
+			~FAttribute();
+
+			virtual void Dump(int32 Indent) const;
 
-			ast_attribute* attribute;
-			ast_attribute_arg* attribute_arg;
+			const TCHAR* Name;
+			TLinearArray<FAttributeArgument*> Arguments;
 		};
-	*/
 	}
 }
diff --git a/Engine/Source/Developer/ShaderCompilerCommon/Private/HlslExpressionParser.inl b/Engine/Source/Developer/ShaderCompilerCommon/Private/HlslExpressionParser.inl
--- a/Engine/Source/Developer/ShaderCompilerCommon/Private/HlslExpressionParser.inl
+++ b/Engine/Source/Developer/ShaderCompilerCommon/Private/HlslExpressionParser.inl
@@ -12,18 +12,20 @@
 namespace CrossCompiler
 {
 	struct FSymbolScope;
-	struct FInfo;
+	//struct FInfo;
 
-	EParseResult ComputeExpr(FHlslScanner& Scanner, int32 MinPrec, FInfo& Info, FSymbolScope* SymbolScope, AST::FExpression** OutExpression);
-	EParseResult ParseExpressionList(EHlslToken EndListToken, FHlslScanner& Scanner, FSymbolScope* SymbolScope, EHlslToken NewStartListToken = EHlslToken::Invalid);
+	EParseResult ComputeExpr(FHlslScanner& Scanner, int32 MinPrec, /*FInfo& Info,*/ FSymbolScope* SymbolScope, bool bAllowAssignment, FLinearAllocator* Allocator, AST::FExpression** OutExpression);
+	EParseResult ParseExpressionList(EHlslToken EndListToken, FHlslScanner& Scanner, FSymbolScope* SymbolScope, EHlslToken NewStartListToken, FLinearAllocator* Allocator, AST::FExpression* OutExpression);
 
 	struct FSymbolScope
 	{
 		FSymbolScope* Parent;
-		TSet<FString> Symbols;
-		TArray<FSymbolScope> Children;
 
-		FSymbolScope(FSymbolScope* InParent) : Parent(InParent) {}
+		TSet/*TLinearSet*/<FString> Symbols;
+		TLinearArray<FSymbolScope> Children;
+
+		FSymbolScope(FLinearAllocator* InAllocator, FSymbolScope* InParent) : Parent(InParent), Children(InAllocator)/*, Symbols(InAllocator)*/ {}
+		~FSymbolScope() {}
 
 		void Add(const FString& Type)
 		{
@@ -51,11 +53,11 @@ namespace CrossCompiler
 		FSymbolScope* Original;
 		FSymbolScope** Current;
 
-		FCreateSymbolScope(FSymbolScope** InCurrent) :
+		FCreateSymbolScope(FLinearAllocator* InAllocator, FSymbolScope** InCurrent) :
 			Current(InCurrent)
 		{
 			Original = *InCurrent;
-			auto* NewScope = new(Original->Children) FSymbolScope(Original);
+			auto* NewScope = new(Original->Children) FSymbolScope(InAllocator, Original);
 			*Current = NewScope;
 		}
 
@@ -65,6 +67,7 @@ namespace CrossCompiler
 		}
 	};
 
+/*
 	struct FInfo
 	{
 		int32 Indent;
@@ -111,7 +114,7 @@ namespace CrossCompiler
 		}
 
 		FInfo& Info;
-	};
+	};*/
 
 	enum ETypeFlags
 	{
@@ -121,14 +124,21 @@ namespace CrossCompiler
 		ETF_USER_TYPES				= 0x0008,
 	};
 
-	EParseResult ParseGeneralType(EHlslToken Token, int32 TypeFlags)
+	EParseResult ParseGeneralType(const FHlslToken* Token, int32 TypeFlags, FLinearAllocator* Allocator, AST::FTypeSpecifier** OutSpecifier)
 	{
-		switch (Token)
+		if (!Token)
+		{
+			return EParseResult::Error;
+		}
+
+		bool bMatched = false;
+		const TCHAR* InnerType = nullptr;
+		switch (Token->Token)
 		{
 		case EHlslToken::Void:
 			if (TypeFlags & ETF_VOID)
 			{
-				return EParseResult::Matched;
+				bMatched = true;
 			}
 			break;
 
@@ -243,7 +253,7 @@ namespace CrossCompiler
 		case EHlslToken::Float4x4:
 			if (TypeFlags & ETF_BUILTIN_NUMERIC)
 			{
-				return EParseResult::Matched;
+				bMatched = true;
 			}
 			break;
 
@@ -281,19 +291,29 @@ namespace CrossCompiler
 		case EHlslToken::SamplerComparisonState:
 			if (TypeFlags & ETF_SAMPLER_TEXTURE_BUFFER)
 			{
-				return EParseResult::Matched;
+				bMatched = true;
+				InnerType = TEXT("float4");
 			}
 			break;
 		}
-		
+
+		if (bMatched)
+		{
+			auto* Type = new(Allocator) AST::FTypeSpecifier(Allocator, Token->SourceInfo);
+			Type->TypeName = Allocator->Strdup(Token->String);
+			Type->InnerType = InnerType;
+			*OutSpecifier = Type;
+			return EParseResult::Matched;
+		}
+
 		return EParseResult::NotMatched;
 	}
 
-	EParseResult ParseGeneralType(const FHlslToken* Token, int32 TypeFlags, FSymbolScope* SymbolScope)
+	EParseResult ParseGeneralType(const FHlslToken* Token, int32 TypeFlags, FSymbolScope* SymbolScope, FLinearAllocator* Allocator, AST::FTypeSpecifier** OutSpecifier)
 	{
 		if (Token)
 		{
-			if (ParseGeneralType(Token->Token, TypeFlags) == EParseResult::Matched)
+			if (ParseGeneralType(Token, TypeFlags, Allocator, OutSpecifier) == EParseResult::Matched)
 			{
 				return EParseResult::Matched;
 			}
@@ -303,6 +323,9 @@ namespace CrossCompiler
 				check(SymbolScope);
 				if (Token->Token == EHlslToken::Identifier && FSymbolScope::FindType(SymbolScope, Token->String))
 				{
+					auto* Type = new(Allocator) AST::FTypeSpecifier(Allocator, Token->SourceInfo);
+					Type->TypeName = Allocator->Strdup(Token->String);
+					*OutSpecifier = Type;
 					return EParseResult::Matched;
 				}
 			}
@@ -313,10 +336,10 @@ namespace CrossCompiler
 		return EParseResult::Error;
 	}
 
-	EParseResult ParseGeneralType(FHlslScanner& Scanner, int32 TypeFlags, FSymbolScope* SymbolScope = nullptr)
+	EParseResult ParseGeneralType(FHlslScanner& Scanner, int32 TypeFlags, FSymbolScope* SymbolScope, FLinearAllocator* Allocator, AST::FTypeSpecifier** OutSpecifier)
 	{
 		auto* Token = Scanner.PeekToken();
-		if (ParseGeneralType(Token, TypeFlags, SymbolScope) == EParseResult::Matched)
+		if (ParseGeneralType(Token, TypeFlags, SymbolScope, Allocator, OutSpecifier) == EParseResult::Matched)
 		{
 			Scanner.Advance();
 			return EParseResult::Matched;
@@ -326,7 +349,7 @@ namespace CrossCompiler
 	}
 
 	// Unary!(Unary-(Unary+())) would have ! as Top, and + as Inner
-	EParseResult MatchUnaryOperator(FHlslScanner& Scanner, FInfo& Info, FSymbolScope* SymbolScope, AST::FExpression** OuterExpression, AST::FExpression** InnerExpression)
+	EParseResult MatchUnaryOperator(FHlslScanner& Scanner, /*FInfo& Info,*/ FSymbolScope* SymbolScope, FLinearAllocator* Allocator, AST::FExpression** OuterExpression, AST::FExpression** InnerExpression)
 	{
 		bool bFoundAny = false;
 		bool bTryAgain = true;
@@ -379,14 +402,30 @@ namespace CrossCompiler
 			{
 				const auto* Peek1 = Scanner.PeekToken(1);
 				const auto* Peek2 = Scanner.PeekToken(2);
-				if (Peek1 && ParseGeneralType(Peek1, ETF_BUILTIN_NUMERIC | ETF_USER_TYPES, SymbolScope) == EParseResult::Matched && Peek2 && Peek2->Token == EHlslToken::RightParenthesis)
+				AST::FTypeSpecifier* TypeSpecifier = nullptr;
+				if (Peek1 && ParseGeneralType(Peek1, ETF_BUILTIN_NUMERIC | ETF_USER_TYPES, SymbolScope, Allocator, &TypeSpecifier) == EParseResult::Matched && Peek2 && Peek2->Token == EHlslToken::RightParenthesis)
 				{
 					// Cast
-check(0);
 					Scanner.Advance();
 					Scanner.Advance();
 					Scanner.Advance();
 					bFoundAny = true;
+
+					auto* Expression = new(Allocator) AST::FUnaryExpression(Allocator, AST::EOperators::TypeCast, nullptr, Token->SourceInfo);
+					if (PrevExpression)
+					{
+						PrevExpression->SubExpressions[0] = Expression;
+					}
+
+					Expression->TypeSpecifier = TypeSpecifier;
+
+					if (!*OuterExpression)
+					{
+						*OuterExpression = Expression;
+					}
+
+					PrevExpression = Expression;
+					continue;
 				}
 				else
 				{
@@ -400,7 +439,7 @@ check(0);
 				return bFoundAny ? EParseResult::Matched : EParseResult::NotMatched;
 			}
 
-			auto* Expression = new AST::FUnaryExpression(Operator, nullptr, Token->SourceInfo);
+			auto* Expression = new(Allocator) AST::FUnaryExpression(Allocator, Operator, nullptr, Token->SourceInfo);
 			if (PrevExpression)
 			{
 				PrevExpression->SubExpressions[0] = Expression;
@@ -418,10 +457,130 @@ check(0);
 		return EParseResult::Error;
 	}
 
-	EParseResult ComputeAtom(FHlslScanner& Scanner, FInfo& Info, FSymbolScope* SymbolScope, AST::FExpression** OutExpression)
+	EParseResult MatchSuffixOperator(FHlslScanner& Scanner, /*FInfo& Info,*/ FSymbolScope* SymbolScope, bool bAllowAssignment, FLinearAllocator* Allocator, AST::FExpression** InOutExpression)
+	{
+		bool bFoundAny = false;
+		bool bTryAgain = true;
+		AST::FExpression*& PrevExpression = *InOutExpression;
+		while (Scanner.HasMoreTokens() && bTryAgain)
+		{
+			auto* Token = Scanner.GetCurrentToken();
+			AST::EOperators Operator = AST::EOperators::Plus;
+
+			switch (Token->Token)
+			{
+			case EHlslToken::LeftSquareBracket:
+			{
+				Scanner.Advance();
+				AST::FExpression* ArrayIndex = nullptr;
+				auto Result = ComputeExpr(Scanner, 1, /*Info,*/ SymbolScope, bAllowAssignment, Allocator, &ArrayIndex);
+				if (Result != EParseResult::Matched)
+				{
+					Scanner.SourceError(TEXT("Expected expression!"));
+					return EParseResult::Error;
+				}
+
+				if (!Scanner.MatchToken(EHlslToken::RightSquareBracket))
+				{
+					Scanner.SourceError(TEXT("Expected ']'!"));
+					return EParseResult::Error;
+				}
+
+				auto* ArrayIndexExpression = new(Allocator) AST::FBinaryExpression(Allocator, AST::EOperators::ArrayIndex, PrevExpression, ArrayIndex, Token->SourceInfo);
+				PrevExpression = ArrayIndexExpression;
+				bFoundAny = true;
+			}
+				break;
+			case EHlslToken::Dot:
+			{
+				Scanner.Advance();
+				const auto* Identifier = Scanner.GetCurrentToken();
+				if (!Scanner.MatchToken(EHlslToken::Identifier))
+				{
+					Scanner.SourceError(TEXT("Expected identifier for member or swizzle!"));
+					return EParseResult::Error;
+				}
+				auto* FieldExpression = new(Allocator) AST::FUnaryExpression(Allocator, AST::EOperators::FieldSelection, PrevExpression, Token->SourceInfo);
+				FieldExpression->Identifier = Allocator->Strdup(Identifier->String);
+				PrevExpression = FieldExpression;
+				bFoundAny = true;
+			}
+				break;
+			case EHlslToken::LeftParenthesis:
+			{
+				Scanner.Advance();
+
+				// Function Call
+				auto* FunctionCall = new(Allocator) AST::FFunctionExpression(Allocator, Token->SourceInfo, PrevExpression);
+				auto Result = ParseExpressionList(EHlslToken::RightParenthesis, Scanner, SymbolScope, EHlslToken::Invalid, Allocator, FunctionCall);
+				if (Result != EParseResult::Matched)
+				{
+					Scanner.SourceError(TEXT("Expected ')'!"));
+					return EParseResult::Error;
+				}
+
+				PrevExpression = FunctionCall;
+
+				bFoundAny = true;
+			}
+				break;
+			case EHlslToken::PlusPlus:
+			{
+				Scanner.Advance();
+				auto* IncExpression = new(Allocator) AST::FUnaryExpression(Allocator, AST::EOperators::PostInc, PrevExpression, Token->SourceInfo);
+				PrevExpression = IncExpression;
+				bFoundAny = true;
+			}
+				break;
+			case EHlslToken::MinusMinus:
+			{
+				Scanner.Advance();
+				auto* DecExpression = new(Allocator) AST::FUnaryExpression(Allocator, AST::EOperators::PostDec, PrevExpression, Token->SourceInfo);
+				PrevExpression = DecExpression;
+				bFoundAny = true;
+			}
+				break;
+			case EHlslToken::Question:
+			{
+				Scanner.Advance();
+				AST::FExpression* Left = nullptr;
+				if (ComputeExpr(Scanner, 0, /*Info,*/ SymbolScope, true, Allocator, &Left) != EParseResult::Matched)
+				{
+					Scanner.SourceError(TEXT("Expected expression!"));
+					return EParseResult::Error;
+				}
+				if (!Scanner.MatchToken(EHlslToken::Colon))
+				{
+					Scanner.SourceError(TEXT("Expected ':'!"));
+					return EParseResult::Error;
+				}
+				AST::FExpression* Right = nullptr;
+				if (ComputeExpr(Scanner, 0, /*Info,*/ SymbolScope, true, Allocator, &Right) != EParseResult::Matched)
+				{
+					Scanner.SourceError(TEXT("Expected expression!"));
+					return EParseResult::Error;
+				}
+
+				auto* Ternary = new(Allocator) AST::FExpression(Allocator, AST::EOperators::Conditional, PrevExpression, Left, Right, Token->SourceInfo);
+				PrevExpression = Ternary;
+				bFoundAny = true;
+			}
+				break;
+
+			default:
+				bTryAgain = false;
+				break;
+			}
+		}
+
+		*InOutExpression = PrevExpression;
+		return bFoundAny ? EParseResult::Matched : EParseResult::NotMatched;
+	}
+
+	EParseResult ComputeAtom(FHlslScanner& Scanner, /*FInfo& Info,*/ FSymbolScope* SymbolScope, bool bAllowAssignment, FLinearAllocator* Allocator, AST::FExpression** OutExpression)
 	{
 		AST::FExpression* InnerUnaryExpression = nullptr;
-		auto UnaryResult = MatchUnaryOperator(Scanner, Info, SymbolScope, OutExpression, &InnerUnaryExpression);
+		auto UnaryResult = MatchUnaryOperator(Scanner, /*Info,*/ SymbolScope, Allocator, OutExpression, &InnerUnaryExpression);
 		auto* Token = Scanner.GetCurrentToken();
 		if (!Token || UnaryResult == EParseResult::Error)
 		{
@@ -433,26 +592,26 @@ check(0);
 		{
 		case EHlslToken::BoolConstant:
 			Scanner.Advance();
-			AtomExpression = new AST::FUnaryExpression(AST::EOperators::BoolConstant, nullptr, Token->SourceInfo);
+			AtomExpression = new(Allocator) AST::FUnaryExpression(Allocator, AST::EOperators::BoolConstant, nullptr, Token->SourceInfo);
 			AtomExpression->BoolConstant = Token->UnsignedInteger != 0;
 			break;
 
 		case EHlslToken::UnsignedIntegerConstant:
 			Scanner.Advance();
-			AtomExpression = new AST::FUnaryExpression(AST::EOperators::UintConstant, nullptr, Token->SourceInfo);
+			AtomExpression = new(Allocator) AST::FUnaryExpression(Allocator, AST::EOperators::UintConstant, nullptr, Token->SourceInfo);
 			AtomExpression->UintConstant = Token->UnsignedInteger;
 			break;
 
 		case EHlslToken::FloatConstant:
 			Scanner.Advance();
-			AtomExpression = new AST::FUnaryExpression(AST::EOperators::FloatConstant, nullptr, Token->SourceInfo);
+			AtomExpression = new(Allocator) AST::FUnaryExpression(Allocator, AST::EOperators::FloatConstant, nullptr, Token->SourceInfo);
 			AtomExpression->FloatConstant = Token->Float;
 			break;
 
 		case EHlslToken::Identifier:
 			Scanner.Advance();
-			AtomExpression = new AST::FUnaryExpression(AST::EOperators::Identifier, nullptr, Token->SourceInfo);
-			AtomExpression->Identifier = Token->String;
+			AtomExpression = new(Allocator) AST::FUnaryExpression(Allocator, AST::EOperators::Identifier, nullptr, Token->SourceInfo);
+			AtomExpression->Identifier = Allocator->Strdup(Token->String);
 			break;
 
 		case EHlslToken::LeftParenthesis:
@@ -463,7 +622,7 @@ check(0);
 			const auto* Peek1 = Scanner.PeekToken(0);
 			const auto* Peek2 = Scanner.PeekToken(1);
 			// Parenthesis expression
-			if (ComputeExpr(Scanner, 1, Info, SymbolScope, &AtomExpression) != EParseResult::Matched)
+			if (ComputeExpr(Scanner, 1, /*Info,*/ SymbolScope, bAllowAssignment, Allocator, &AtomExpression) != EParseResult::Matched)
 			{
 				Scanner.SourceError(TEXT("Expected expression!"));
 				return EParseResult::Error;
@@ -478,24 +637,33 @@ check(0);
 			break;
 
 		default:
+		{
+			AST::FTypeSpecifier* TypeSpecifier = nullptr;
+
 			// Grrr handle Sampler as a variable name... This is safe here since Declarations are always handled first
-			if (ParseGeneralType(Scanner, ETF_SAMPLER_TEXTURE_BUFFER) == EParseResult::Matched)
+			if (ParseGeneralType(Scanner, ETF_SAMPLER_TEXTURE_BUFFER, nullptr, Allocator, &TypeSpecifier) == EParseResult::Matched)
 			{
-				check(0);
+				//@todo-rco: Check this var exists on the symnbol table
+				AtomExpression = new(Allocator) AST::FUnaryExpression(Allocator, AST::EOperators::Identifier, nullptr, TypeSpecifier->SourceInfo);
+				AtomExpression->Identifier = TypeSpecifier->TypeName;
 				break;
 			}
 			// Handle float3(x,y,z)
-			else if (ParseGeneralType(Scanner, ETF_BUILTIN_NUMERIC) == EParseResult::Matched)
+			else if (ParseGeneralType(Scanner, ETF_BUILTIN_NUMERIC, nullptr, Allocator, &TypeSpecifier) == EParseResult::Matched)
 			{
-				check(0);
 				if (Scanner.MatchToken(EHlslToken::LeftParenthesis))
 				{
-					auto Result = ParseExpressionList(EHlslToken::RightParenthesis, Scanner, SymbolScope);
+					auto* TypeExpression = new(Allocator) AST::FUnaryExpression(Allocator, AST::EOperators::Identifier, nullptr, TypeSpecifier->SourceInfo);
+					TypeExpression->Identifier = TypeSpecifier->TypeName;
+					auto* FunctionCall = new(Allocator) AST::FFunctionExpression(Allocator, Token->SourceInfo, TypeExpression);
+					auto Result = ParseExpressionList(EHlslToken::RightParenthesis, Scanner, SymbolScope, EHlslToken::Invalid, Allocator, FunctionCall);
 					if (Result != EParseResult::Matched)
 					{
 						Scanner.SourceError(TEXT("Unexpected type in numeric constructor!"));
 						return EParseResult::Error;
 					}
+
+					AtomExpression = FunctionCall;
 				}
 				else
 				{
@@ -515,82 +683,16 @@ check(0);
 				return EParseResult::NotMatched;
 			}
 		}
+			break;
+		}
 
 		check(AtomExpression);
 
-		// Suffix
-		while (Scanner.HasMoreTokens())
+		auto SuffixResult = MatchSuffixOperator(Scanner, /*Info,*/ SymbolScope, bAllowAssignment, Allocator, &AtomExpression);
+		//auto* Token = Scanner.GetCurrentToken();
+		if (/*!Token || */SuffixResult == EParseResult::Error)
 		{
-			auto* Token = Scanner.GetCurrentToken();
-			if (Scanner.MatchToken(EHlslToken::LeftSquareBracket))
-			{
-check(0);
-/*
-				auto Result = ComputeExpr(Scanner, 1, Info, SymbolScope);
-				if (Result != EParseResult::Matched)
-				{
-					Scanner.SourceError(TEXT("Expected expression!"));
-					return EParseResult::Error;
-				}
-
-				if (!Scanner.MatchToken(EHlslToken::RightSquareBracket))
-				{
-					Scanner.SourceError(TEXT("Expected ']'!"));
-					return EParseResult::Error;
-				}
-*/
-			}
-			else if (Scanner.MatchToken(EHlslToken::Dot))
-			{
-				check(0);
-				if (!Scanner.MatchToken(EHlslToken::Identifier))
-				{
-					Scanner.SourceError(TEXT("Expected identifier for member or swizzle!"));
-					return EParseResult::Error;
-				}
-			}
-			else if (Scanner.MatchToken(EHlslToken::LeftParenthesis))
-			{
-				check(0);
-				// Function Call
-				auto Result = ParseExpressionList(EHlslToken::RightParenthesis, Scanner, SymbolScope);
-				if (Result != EParseResult::Matched)
-				{
-					Scanner.SourceError(TEXT("Expected ')'!"));
-					return EParseResult::Error;
-				}
-			}
-			else if (Scanner.MatchToken(EHlslToken::PlusPlus) || Scanner.MatchToken(EHlslToken::MinusMinus))
-			{
-				check(0);
-			}
-			else if (Scanner.MatchToken(EHlslToken::Question))
-			{
-				check(0);
-/*
-				if (ComputeExpr(Scanner, 0, Info, SymbolScope) != EParseResult::Matched)
-				{
-					Scanner.SourceError(TEXT("Expected expression!"));
-					return EParseResult::Error;
-				}
-				if (!Scanner.MatchToken(EHlslToken::Colon))
-				{
-					Scanner.SourceError(TEXT("Expected ':'!"));
-					return EParseResult::Error;
-				}
-				//FPlatformMisc::LowLevelOutputDebugStringf(TEXT("\t:"));
-				if (ComputeExpr(Scanner, 0, Info, SymbolScope) != EParseResult::Matched)
-				{
-					Scanner.SourceError(TEXT("Expected expression!"));
-					return EParseResult::Error;
-				}
-				break;
-*/
-			}
-			else
-			{
-				break;
-			}
+			return EParseResult::Error;
 		}
 
 		// Patch unary if necessary
@@ -681,10 +783,37 @@ check(0);
 		return GetPrecedence(Token) > 0;
 	}
 
-	EParseResult ComputeExpr(FHlslScanner& Scanner, int32 MinPrec, FInfo& Info, FSymbolScope* SymbolScope, AST::FExpression** OutExpression)
+	bool IsAssignmentOperator(const FHlslToken* Token)
+	{
+		if (Token)
+		{
+			switch (Token->Token)
+			{
+			case EHlslToken::Equal:
+			case EHlslToken::PlusEqual:
+			case EHlslToken::MinusEqual:
+			case EHlslToken::TimesEqual:
+			case EHlslToken::DivEqual:
+			case EHlslToken::ModEqual:
+			case EHlslToken::GreaterGreaterEqual:
+			case EHlslToken::LowerLowerEqual:
+			case EHlslToken::AndEqual:
+			case EHlslToken::OrEqual:
+			case EHlslToken::XorEqual:
+				return true;
+
+			default:
+				break;
+			}
+		}
+
+		return false;
+	}
+
+	EParseResult ComputeExpr(FHlslScanner& Scanner, int32 MinPrec, /*FInfo& Info,*/ FSymbolScope* SymbolScope, bool bAllowAssignment, FLinearAllocator* Allocator, AST::FExpression** OutExpression)
 	{
 		auto OriginalToken = Scanner.GetCurrentTokenIndex();
-		FInfoIndentScope Scope(Info);
+		//FInfoIndentScope Scope(Info);
 		/*
 		// Precedence Climbing
 		// http://eli.thegreenplace.net/2012/08/02/parsing-expressions-by-precedence-climbing
@@ -702,8 +831,8 @@ check(0);
 
 			  return result
 		*/
-		Info.PrintWithTabs(FString::Printf(TEXT("Compute Expr %d\n"), MinPrec));
-		auto Result = ComputeAtom(Scanner, Info, SymbolScope, OutExpression);
+		//Info.PrintWithTabs(FString::Printf(TEXT("Compute Expr %d\n"), MinPrec));
+		auto Result = ComputeAtom(Scanner, /*Info,*/ SymbolScope, bAllowAssignment, Allocator, OutExpression);
 		if (Result != EParseResult::Matched)
 		{
 			return Result;
@@ -713,7 +842,7 @@ check(0);
 		{
 			auto* Token = Scanner.GetCurrentToken();
 			int32 Precedence = GetPrecedence(Token);
-			if (!Token || !IsBinaryOperator(Token) || Precedence < MinPrec)
+			if (!Token || !IsBinaryOperator(Token) || Precedence < MinPrec || (!bAllowAssignment && IsAssignmentOperator(Token)))
 			{
 				break;
 			}
@@ -721,7 +850,7 @@ check(0);
 			Scanner.Advance();
 			auto NextMinPrec = Precedence + 1;
 			AST::FExpression* RHSExpression = nullptr;
-			Result = ComputeExpr(Scanner, NextMinPrec, Info, SymbolScope, &RHSExpression);
+			Result = ComputeExpr(Scanner, NextMinPrec, /*Info,*/ SymbolScope, bAllowAssignment, Allocator, &RHSExpression);
 			if (Result == EParseResult::Error)
 			{
 				return EParseResult::Error;
@@ -732,7 +861,7 @@ check(0);
 			}
 			check(RHSExpression);
 			auto BinaryOperator = AST::TokenToASTOperator(Token->Token);
-			*OutExpression = new AST::FBinaryExpression(BinaryOperator, *OutExpression, RHSExpression, Token->SourceInfo);
+			*OutExpression = new(Allocator) AST::FBinaryExpression(Allocator, BinaryOperator, *OutExpression, RHSExpression, Token->SourceInfo);
 		}
 		while (Scanner.HasMoreTokens());
 
@@ -744,15 +873,15 @@ check(0);
 		return EParseResult::Matched;
 	}
 
-	EParseResult ParseExpression(FHlslScanner& Scanner, FSymbolScope* SymbolScope)
+	EParseResult ParseExpression(FHlslScanner& Scanner, FSymbolScope* SymbolScope, bool bAllowAssignment, FLinearAllocator* Allocator, AST::FExpression** OutExpression)
 	{
-		FInfo Info(!true);
-		AST::FExpression* Expression = nullptr;
-		return ComputeExpr(Scanner, 0, Info, SymbolScope, &Expression);
+		/*FInfo Info(!true);*/
+		return ComputeExpr(Scanner, 0, /*Info,*/ SymbolScope, bAllowAssignment, Allocator, OutExpression);
 	}
 
-	EParseResult ParseExpressionList(EHlslToken EndListToken, FHlslScanner& Scanner, FSymbolScope* SymbolScope, EHlslToken NewStartListToken)
+	EParseResult ParseExpressionList(EHlslToken EndListToken, FHlslScanner& Scanner, FSymbolScope* SymbolScope, EHlslToken NewStartListToken, FLinearAllocator* Allocator, AST::FExpression* OutExpression)
 	{
+		check(OutExpression);
 		while (Scanner.HasMoreTokens())
 		{
 			const auto* Token = Scanner.PeekToken();
@@ -764,15 +893,19 @@ check(0);
 			else if (NewStartListToken != EHlslToken::Invalid && Token->Token == NewStartListToken)
 			{
 				Scanner.Advance();
-				auto Result = ParseExpressionList(EndListToken, Scanner, SymbolScope, NewStartListToken);
+				auto* SubExpression = new(Allocator) AST::FInitializerListExpression(Allocator, Token->SourceInfo);
+				auto Result = ParseExpressionList(EndListToken, Scanner, SymbolScope, NewStartListToken, Allocator, SubExpression);
 				if (Result != EParseResult::Matched)
 				{
 					return Result;
 				}
+
+				OutExpression->Expressions.Add(SubExpression);
 			}
 			else
 			{
-				auto Result = ParseExpression(Scanner, SymbolScope);
+				AST::FExpression* Expression = nullptr;
+				auto Result = ParseExpression(Scanner, SymbolScope, true, Allocator, &Expression);
 				if (Result == EParseResult::Error)
 				{
 					Scanner.SourceError(TEXT("Invalid expression list\n"));
@@ -783,6 +916,8 @@ check(0);
 					Scanner.SourceError(TEXT("Expected expression\n"));
 					return EParseResult::Error;
 				}
+
+				OutExpression->Expressions.Add(Expression);
 			}
 
 			if (Scanner.MatchToken(EHlslToken::Comma))
diff --git a/Engine/Source/Developer/ShaderCompilerCommon/Private/HlslLexer.cpp b/Engine/Source/Developer/ShaderCompilerCommon/Private/HlslLexer.cpp
--- a/Engine/Source/Developer/ShaderCompilerCommon/Private/HlslLexer.cpp
+++ b/Engine/Source/Developer/ShaderCompilerCommon/Private/HlslLexer.cpp
@@ -157,234 +157,241 @@ namespace CrossCompiler
 		return false;
 	}
 
-	static void InitializeTokenTables()
+	namespace Tokens
 	{
-		// Math
-		InsertToken(TEXT("+"), EHlslToken::Plus);
-		InsertToken(TEXT("+="), EHlslToken::PlusEqual);
-		InsertToken(TEXT("-"), EHlslToken::Minus);
-		InsertToken(TEXT("-="), EHlslToken::MinusEqual);
-		InsertToken(TEXT("*"), EHlslToken::Times);
-		InsertToken(TEXT("*="), EHlslToken::TimesEqual);
-		InsertToken(TEXT("/"), EHlslToken::Div);
-		InsertToken(TEXT("/="), EHlslToken::DivEqual);
-		InsertToken(TEXT("%"), EHlslToken::Mod);
-		InsertToken(TEXT("%="), EHlslToken::ModEqual);
-		InsertToken(TEXT("("), EHlslToken::LeftParenthesis);
-		InsertToken(TEXT(")"), EHlslToken::RightParenthesis);
-
-		// Logical
-		InsertToken(TEXT("=="), EHlslToken::EqualEqual);
-		InsertToken(TEXT("!="), EHlslToken::NotEqual);
-		InsertToken(TEXT("<"), EHlslToken::Lower);
-		InsertToken(TEXT("<="), EHlslToken::LowerEqual);
-		InsertToken(TEXT(">"), EHlslToken::Greater);
-		InsertToken(TEXT(">="), EHlslToken::GreaterEqual);
-		InsertToken(TEXT("&&"), EHlslToken::AndAnd);
-		InsertToken(TEXT("||"), EHlslToken::OrOr);
-
-		// Bit
-		InsertToken(TEXT("<<"), EHlslToken::LowerLower);
-		InsertToken(TEXT("<<="), EHlslToken::LowerLowerEqual);
-		InsertToken(TEXT(">>"), EHlslToken::GreaterGreater);
-		InsertToken(TEXT(">>="), EHlslToken::GreaterGreaterEqual);
-		InsertToken(TEXT("&"), EHlslToken::And);
-		InsertToken(TEXT("&="), EHlslToken::And);
-		InsertToken(TEXT("|"), EHlslToken::Or);
-		InsertToken(TEXT("|="), EHlslToken::OrEqual);
-		InsertToken(TEXT("^"), EHlslToken::Xor);
-		InsertToken(TEXT("^="), EHlslToken::XorEqual);
-		InsertToken(TEXT("!"), EHlslToken::Not);
-		InsertToken(TEXT("~"), EHlslToken::Neg);
-
-		// Statements/Keywords
-		InsertToken(TEXT("="), EHlslToken::Equal);
-		InsertToken(TEXT("{"), EHlslToken::LeftBrace);
-		InsertToken(TEXT("}"), EHlslToken::RightBrace);
-		InsertToken(TEXT(";"), EHlslToken::Semicolon);
-		InsertToken(TEXT("if"), EHlslToken::If);
-		InsertToken(TEXT("else"), EHlslToken::Else);
-		InsertToken(TEXT("for"), EHlslToken::For);
-		InsertToken(TEXT("while"), EHlslToken::While);
-		InsertToken(TEXT("do"), EHlslToken::Do);
-		InsertToken(TEXT("return"), EHlslToken::Return);
-		InsertToken(TEXT("switch"), EHlslToken::Switch);
-		InsertToken(TEXT("case"), EHlslToken::Case);
-		InsertToken(TEXT("break"), EHlslToken::Break);
-		InsertToken(TEXT("default"), EHlslToken::Default);
-		InsertToken(TEXT("goto"), EHlslToken::Goto);
-
-		// Unary
-		InsertToken(TEXT("++"), EHlslToken::PlusPlus);
-		InsertToken(TEXT("--"), EHlslToken::MinusMinus);
-
-		// Types
-		InsertToken(TEXT("void"), EHlslToken::Void);
-		InsertToken(TEXT("const"), EHlslToken::Const);
-
-		InsertToken(TEXT("bool"), EHlslToken::Bool);
-		InsertToken(TEXT("bool1"), EHlslToken::Bool1);
-		InsertToken(TEXT("bool2"), EHlslToken::Bool2);
-		InsertToken(TEXT("bool3"), EHlslToken::Bool3);
-		InsertToken(TEXT("bool4"), EHlslToken::Bool4);
-		InsertToken(TEXT("bool1x1"), EHlslToken::Bool1x1);
-		InsertToken(TEXT("bool2x1"), EHlslToken::Bool2x1);
-		InsertToken(TEXT("bool3x1"), EHlslToken::Bool3x1);
-		InsertToken(TEXT("bool4x1"), EHlslToken::Bool4x1);
-		InsertToken(TEXT("bool1x2"), EHlslToken::Bool1x2);
-		InsertToken(TEXT("bool2x2"), EHlslToken::Bool2x2);
-		InsertToken(TEXT("bool3x2"), EHlslToken::Bool3x2);
-		InsertToken(TEXT("bool4x2"), EHlslToken::Bool4x2);
-		InsertToken(TEXT("bool1x3"), EHlslToken::Bool1x3);
-		InsertToken(TEXT("bool2x3"), EHlslToken::Bool2x3);
-		InsertToken(TEXT("bool3x3"), EHlslToken::Bool3x3);
-		InsertToken(TEXT("bool4x3"), EHlslToken::Bool4x3);
-		InsertToken(TEXT("bool1x4"), EHlslToken::Bool1x4);
-		InsertToken(TEXT("bool2x4"), EHlslToken::Bool2x4);
-		InsertToken(TEXT("bool3x4"), EHlslToken::Bool3x4);
-		InsertToken(TEXT("bool4x4"), EHlslToken::Bool4x4);
-
-		InsertToken(TEXT("int"), EHlslToken::Int);
-		InsertToken(TEXT("int1"), EHlslToken::Int1);
-		InsertToken(TEXT("int2"), EHlslToken::Int2);
-		InsertToken(TEXT("int3"), EHlslToken::Int3);
-		InsertToken(TEXT("int4"), EHlslToken::Int4);
-		InsertToken(TEXT("int1x1"), EHlslToken::Int1x1);
-		InsertToken(TEXT("int2x1"), EHlslToken::Int2x1);
-		InsertToken(TEXT("int3x1"), EHlslToken::Int3x1);
-		InsertToken(TEXT("int4x1"), EHlslToken::Int4x1);
-		InsertToken(TEXT("int1x2"), EHlslToken::Int1x2);
-		InsertToken(TEXT("int2x2"), EHlslToken::Int2x2);
-		InsertToken(TEXT("int3x2"), EHlslToken::Int3x2);
-		InsertToken(TEXT("int4x2"), EHlslToken::Int4x2);
-		InsertToken(TEXT("int1x3"), EHlslToken::Int1x3);
-		InsertToken(TEXT("int2x3"), EHlslToken::Int2x3);
-		InsertToken(TEXT("int3x3"), EHlslToken::Int3x3);
-		InsertToken(TEXT("int4x3"), EHlslToken::Int4x3);
-		InsertToken(TEXT("int1x4"), EHlslToken::Int1x4);
-		InsertToken(TEXT("int2x4"), EHlslToken::Int2x4);
-		InsertToken(TEXT("int3x4"), EHlslToken::Int3x4);
-		InsertToken(TEXT("int4x4"), EHlslToken::Int4x4);
-
-		InsertToken(TEXT("uint"), EHlslToken::Uint);
-		InsertToken(TEXT("uint1"), EHlslToken::Uint1);
-		InsertToken(TEXT("uint2"), EHlslToken::Uint2);
-		InsertToken(TEXT("uint3"), EHlslToken::Uint3);
-		InsertToken(TEXT("uint4"), EHlslToken::Uint4);
-		InsertToken(TEXT("uint1x1"), EHlslToken::Uint1x1);
-		InsertToken(TEXT("uint2x1"), EHlslToken::Uint2x1);
-		InsertToken(TEXT("uint3x1"), EHlslToken::Uint3x1);
-		InsertToken(TEXT("uint4x1"), EHlslToken::Uint4x1);
-		InsertToken(TEXT("uint1x2"), EHlslToken::Uint1x2);
-		InsertToken(TEXT("uint2x2"), EHlslToken::Uint2x2);
-		InsertToken(TEXT("uint3x2"), EHlslToken::Uint3x2);
-		InsertToken(TEXT("uint4x2"), EHlslToken::Uint4x2);
-		InsertToken(TEXT("uint1x3"), EHlslToken::Uint1x3);
-		InsertToken(TEXT("uint2x3"), EHlslToken::Uint2x3);
-		InsertToken(TEXT("uint3x3"), EHlslToken::Uint3x3);
-		InsertToken(TEXT("uint4x3"), EHlslToken::Uint4x3);
-		InsertToken(TEXT("uint1x4"), EHlslToken::Uint1x4);
-		InsertToken(TEXT("uint2x4"), EHlslToken::Uint2x4);
-		InsertToken(TEXT("uint3x4"), EHlslToken::Uint3x4);
-		InsertToken(TEXT("uint4x4"), EHlslToken::Uint4x4);
-
-		InsertToken(TEXT("half"), EHlslToken::Half);
-		InsertToken(TEXT("half1"), EHlslToken::Half1);
-		InsertToken(TEXT("half2"), EHlslToken::Half2);
-		InsertToken(TEXT("half3"), EHlslToken::Half3);
-		InsertToken(TEXT("half4"), EHlslToken::Half4);
-		InsertToken(TEXT("half1x1"), EHlslToken::Half1x1);
-		InsertToken(TEXT("half2x1"), EHlslToken::Half2x1);
-		InsertToken(TEXT("half3x1"), EHlslToken::Half3x1);
-		InsertToken(TEXT("half4x1"), EHlslToken::Half4x1);
-		InsertToken(TEXT("half1x2"), EHlslToken::Half1x2);
-		InsertToken(TEXT("half2x2"), EHlslToken::Half2x2);
-		InsertToken(TEXT("half3x2"), EHlslToken::Half3x2);
-		InsertToken(TEXT("half4x2"), EHlslToken::Half4x2);
-		InsertToken(TEXT("half1x3"), EHlslToken::Half1x3);
-		InsertToken(TEXT("half2x3"), EHlslToken::Half2x3);
-		InsertToken(TEXT("half3x3"), EHlslToken::Half3x3);
-		InsertToken(TEXT("half4x3"), EHlslToken::Half4x3);
-		InsertToken(TEXT("half1x4"), EHlslToken::Half1x4);
-		InsertToken(TEXT("half2x4"), EHlslToken::Half2x4);
-		InsertToken(TEXT("half3x4"), EHlslToken::Half3x4);
-		InsertToken(TEXT("half4x4"), EHlslToken::Half4x4);
-
-		InsertToken(TEXT("float"), EHlslToken::Float);
-		InsertToken(TEXT("float1"), EHlslToken::Float1);
-		InsertToken(TEXT("float2"), EHlslToken::Float2);
-		InsertToken(TEXT("float3"), EHlslToken::Float3);
-		InsertToken(TEXT("float4"), EHlslToken::Float4);
-		InsertToken(TEXT("float1x1"), EHlslToken::Float1x1);
-		InsertToken(TEXT("float2x1"), EHlslToken::Float2x1);
-		InsertToken(TEXT("float3x1"), EHlslToken::Float3x1);
-		InsertToken(TEXT("float4x1"), EHlslToken::Float4x1);
-		InsertToken(TEXT("float1x2"), EHlslToken::Float1x2);
-		InsertToken(TEXT("float2x2"), EHlslToken::Float2x2);
-		InsertToken(TEXT("float3x2"), EHlslToken::Float3x2);
-		InsertToken(TEXT("float4x2"), EHlslToken::Float4x2);
-		InsertToken(TEXT("float1x3"), EHlslToken::Float1x3);
-		InsertToken(TEXT("float2x3"), EHlslToken::Float2x3);
-		InsertToken(TEXT("float3x3"), EHlslToken::Float3x3);
-		InsertToken(TEXT("float4x3"), EHlslToken::Float4x3);
-		InsertToken(TEXT("float1x4"), EHlslToken::Float1x4);
-		InsertToken(TEXT("float2x4"), EHlslToken::Float2x4);
-		InsertToken(TEXT("float3x4"), EHlslToken::Float3x4);
-		InsertToken(TEXT("float4x4"), EHlslToken::Float4x4);
-
-		InsertToken(TEXT("Texture"), EHlslToken::Texture);
-		InsertToken(TEXT("Texture1D"), EHlslToken::Texture1D);
-		InsertToken(TEXT("Texture1DArray"), EHlslToken::Texture1DArray);
-		InsertToken(TEXT("Texture2D"), EHlslToken::Texture2D);
-		InsertToken(TEXT("Texture2DArray"), EHlslToken::Texture2DArray);
-		InsertToken(TEXT("Texture2DMS"), EHlslToken::Texture2DMS);
-		InsertToken(TEXT("Texture2DMSArray"), EHlslToken::Texture2DMSArray);
-		InsertToken(TEXT("Texture3D"), EHlslToken::Texture3D);
-		InsertToken(TEXT("TextureCube"), EHlslToken::TextureCube);
-		InsertToken(TEXT("TextureCubeArray"), EHlslToken::TextureCubeArray);
-
-		InsertToken(TEXT("Sampler"), EHlslToken::Sampler);
-		InsertToken(TEXT("Sampler1D"), EHlslToken::Sampler1D);
-		InsertToken(TEXT("Sampler2D"), EHlslToken::Sampler2D);
-		InsertToken(TEXT("Sampler3D"), EHlslToken::Sampler3D);
-		InsertToken(TEXT("SamplerCube"), EHlslToken::SamplerCube);
-		InsertToken(TEXT("SamplerState"), EHlslToken::SamplerState);
-		InsertToken(TEXT("SamplerComparisonState"), EHlslToken::SamplerComparisonState);
-
-		InsertToken(TEXT("Buffer"), EHlslToken::Buffer);
-		InsertToken(TEXT("AppendStructuredBuffer"), EHlslToken::AppendStructuredBuffer);
-		InsertToken(TEXT("ByteAddressBuffer"), EHlslToken::ByteAddressBuffer);
-		InsertToken(TEXT("ConsumeStructuredBuffer"), EHlslToken::ConsumeStructuredBuffer);
-		InsertToken(TEXT("RWBuffer"), EHlslToken::RWBuffer);
-		InsertToken(TEXT("RWByteAddressBuffer"), EHlslToken::RWByteAddressBuffer);
-		InsertToken(TEXT("RWStructuredBuffer"), EHlslToken::RWStructuredBuffer);
-		InsertToken(TEXT("RWTexture1D"), EHlslToken::RWTexture1D);
-		InsertToken(TEXT("RWTexture1DArray"), EHlslToken::RWTexture1DArray);
-		InsertToken(TEXT("RWTexture2D"), EHlslToken::RWTexture2D);
-		InsertToken(TEXT("RWTexture2DArray"), EHlslToken::RWTexture2DArray);
-		InsertToken(TEXT("RWTexture3D"), EHlslToken::RWTexture3D);
-		InsertToken(TEXT("StructuredBuffer"), EHlslToken::StructuredBuffer);
-		InsertToken(TEXT("InputPatch"), EHlslToken::InputPatch);
-		InsertToken(TEXT("OutputPatch"), EHlslToken::OutputPatch);
-
-		// Modifiers
-		InsertToken(TEXT("in"), EHlslToken::In);
-		InsertToken(TEXT("out"), EHlslToken::Out);
-		InsertToken(TEXT("inout"), EHlslToken::InOut);
-		InsertToken(TEXT("static"), EHlslToken::Static);
-
-		// Misc
-		InsertToken(TEXT("["), EHlslToken::LeftSquareBracket);
-		InsertToken(TEXT("]"), EHlslToken::RightSquareBracket);
-		InsertToken(TEXT("?"), EHlslToken::Question);
-		InsertToken(TEXT(":"), EHlslToken::Colon);
-		InsertToken(TEXT(","), EHlslToken::Comma);
-		InsertToken(TEXT("."), EHlslToken::Dot);
-		InsertToken(TEXT("struct"), EHlslToken::Struct);
-		InsertToken(TEXT("cbuffer"), EHlslToken::CBuffer);
-		InsertToken(TEXT("groupshared"), EHlslToken::GroupShared);
-		InsertToken(TEXT("nointerpolation"), EHlslToken::NoInterpolation);
-		InsertToken(TEXT("row_major"), EHlslToken::RowMajor);
+		static struct FStaticInitializer
+		{
+			FStaticInitializer()
+			{
+				// Math
+				InsertToken(TEXT("+"), EHlslToken::Plus);
+				InsertToken(TEXT("+="), EHlslToken::PlusEqual);
+				InsertToken(TEXT("-"), EHlslToken::Minus);
+				InsertToken(TEXT("-="), EHlslToken::MinusEqual);
+				InsertToken(TEXT("*"), EHlslToken::Times);
+				InsertToken(TEXT("*="), EHlslToken::TimesEqual);
+				InsertToken(TEXT("/"), EHlslToken::Div);
+				InsertToken(TEXT("/="), EHlslToken::DivEqual);
+				InsertToken(TEXT("%"), EHlslToken::Mod);
+				InsertToken(TEXT("%="), EHlslToken::ModEqual);
+				InsertToken(TEXT("("), EHlslToken::LeftParenthesis);
+				InsertToken(TEXT(")"), EHlslToken::RightParenthesis);
+
+				// Logical
+				InsertToken(TEXT("=="), EHlslToken::EqualEqual);
+				InsertToken(TEXT("!="), EHlslToken::NotEqual);
+				InsertToken(TEXT("<"), EHlslToken::Lower);
+				InsertToken(TEXT("<="), EHlslToken::LowerEqual);
+				InsertToken(TEXT(">"), EHlslToken::Greater);
+				InsertToken(TEXT(">="), EHlslToken::GreaterEqual);
+				InsertToken(TEXT("&&"), EHlslToken::AndAnd);
+				InsertToken(TEXT("||"), EHlslToken::OrOr);
+
+				// Bit
+				InsertToken(TEXT("<<"), EHlslToken::LowerLower);
+				InsertToken(TEXT("<<="), EHlslToken::LowerLowerEqual);
+				InsertToken(TEXT(">>"), EHlslToken::GreaterGreater);
+				InsertToken(TEXT(">>="), EHlslToken::GreaterGreaterEqual);
+				InsertToken(TEXT("&"), EHlslToken::And);
+				InsertToken(TEXT("&="), EHlslToken::And);
+				InsertToken(TEXT("|"), EHlslToken::Or);
+				InsertToken(TEXT("|="), EHlslToken::OrEqual);
+				InsertToken(TEXT("^"), EHlslToken::Xor);
+				InsertToken(TEXT("^="), EHlslToken::XorEqual);
+				InsertToken(TEXT("!"), EHlslToken::Not);
+				InsertToken(TEXT("~"), EHlslToken::Neg);
+
+				// Statements/Keywords
+				InsertToken(TEXT("="), EHlslToken::Equal);
+				InsertToken(TEXT("{"), EHlslToken::LeftBrace);
+				InsertToken(TEXT("}"), EHlslToken::RightBrace);
+				InsertToken(TEXT(";"), EHlslToken::Semicolon);
+				InsertToken(TEXT("if"), EHlslToken::If);
+				InsertToken(TEXT("else"), EHlslToken::Else);
+				InsertToken(TEXT("for"), EHlslToken::For);
+				InsertToken(TEXT("while"), EHlslToken::While);
+				InsertToken(TEXT("do"), EHlslToken::Do);
+				InsertToken(TEXT("return"), EHlslToken::Return);
+				InsertToken(TEXT("switch"), EHlslToken::Switch);
+				InsertToken(TEXT("case"), EHlslToken::Case);
+				InsertToken(TEXT("break"), EHlslToken::Break);
+				InsertToken(TEXT("default"), EHlslToken::Default);
+				InsertToken(TEXT("continue"), EHlslToken::Continue);
+				InsertToken(TEXT("goto"), EHlslToken::Goto);
+
+				// Unary
+				InsertToken(TEXT("++"), EHlslToken::PlusPlus);
+				InsertToken(TEXT("--"), EHlslToken::MinusMinus);
+
+				// Types
+				InsertToken(TEXT("void"), EHlslToken::Void);
+				InsertToken(TEXT("const"), EHlslToken::Const);
+
+				InsertToken(TEXT("bool"), EHlslToken::Bool);
+				InsertToken(TEXT("bool1"), EHlslToken::Bool1);
+				InsertToken(TEXT("bool2"), EHlslToken::Bool2);
+				InsertToken(TEXT("bool3"), EHlslToken::Bool3);
+				InsertToken(TEXT("bool4"), EHlslToken::Bool4);
+				InsertToken(TEXT("bool1x1"), EHlslToken::Bool1x1);
+				InsertToken(TEXT("bool2x1"), EHlslToken::Bool2x1);
+				InsertToken(TEXT("bool3x1"), EHlslToken::Bool3x1);
+				InsertToken(TEXT("bool4x1"), EHlslToken::Bool4x1);
+				InsertToken(TEXT("bool1x2"), EHlslToken::Bool1x2);
+				InsertToken(TEXT("bool2x2"), EHlslToken::Bool2x2);
+				InsertToken(TEXT("bool3x2"), EHlslToken::Bool3x2);
+				InsertToken(TEXT("bool4x2"), EHlslToken::Bool4x2);
+				InsertToken(TEXT("bool1x3"), EHlslToken::Bool1x3);
+				InsertToken(TEXT("bool2x3"), EHlslToken::Bool2x3);
+				InsertToken(TEXT("bool3x3"), EHlslToken::Bool3x3);
+				InsertToken(TEXT("bool4x3"), EHlslToken::Bool4x3);
+				InsertToken(TEXT("bool1x4"), EHlslToken::Bool1x4);
+				InsertToken(TEXT("bool2x4"), EHlslToken::Bool2x4);
+				InsertToken(TEXT("bool3x4"), EHlslToken::Bool3x4);
+				InsertToken(TEXT("bool4x4"), EHlslToken::Bool4x4);
+
+				InsertToken(TEXT("int"), EHlslToken::Int);
+				InsertToken(TEXT("int1"), EHlslToken::Int1);
+				InsertToken(TEXT("int2"), EHlslToken::Int2);
+				InsertToken(TEXT("int3"), EHlslToken::Int3);
+				InsertToken(TEXT("int4"), EHlslToken::Int4);
+				InsertToken(TEXT("int1x1"), EHlslToken::Int1x1);
+				InsertToken(TEXT("int2x1"), EHlslToken::Int2x1);
+				InsertToken(TEXT("int3x1"), EHlslToken::Int3x1);
+				InsertToken(TEXT("int4x1"), EHlslToken::Int4x1);
+				InsertToken(TEXT("int1x2"), EHlslToken::Int1x2);
+				InsertToken(TEXT("int2x2"), EHlslToken::Int2x2);
+				InsertToken(TEXT("int3x2"), EHlslToken::Int3x2);
+				InsertToken(TEXT("int4x2"), EHlslToken::Int4x2);
+				InsertToken(TEXT("int1x3"), EHlslToken::Int1x3);
+				InsertToken(TEXT("int2x3"), EHlslToken::Int2x3);
+				InsertToken(TEXT("int3x3"), EHlslToken::Int3x3);
+				InsertToken(TEXT("int4x3"), EHlslToken::Int4x3);
+				InsertToken(TEXT("int1x4"), EHlslToken::Int1x4);
+				InsertToken(TEXT("int2x4"), EHlslToken::Int2x4);
+				InsertToken(TEXT("int3x4"), EHlslToken::Int3x4);
+				InsertToken(TEXT("int4x4"), EHlslToken::Int4x4);
+
+				InsertToken(TEXT("uint"), EHlslToken::Uint);
+				InsertToken(TEXT("uint1"), EHlslToken::Uint1);
+				InsertToken(TEXT("uint2"), EHlslToken::Uint2);
+				InsertToken(TEXT("uint3"), EHlslToken::Uint3);
+				InsertToken(TEXT("uint4"), EHlslToken::Uint4);
+				InsertToken(TEXT("uint1x1"), EHlslToken::Uint1x1);
+				InsertToken(TEXT("uint2x1"), EHlslToken::Uint2x1);
+				InsertToken(TEXT("uint3x1"), EHlslToken::Uint3x1);
+				InsertToken(TEXT("uint4x1"), EHlslToken::Uint4x1);
+				InsertToken(TEXT("uint1x2"), EHlslToken::Uint1x2);
+				InsertToken(TEXT("uint2x2"), EHlslToken::Uint2x2);
+				InsertToken(TEXT("uint3x2"), EHlslToken::Uint3x2);
+				InsertToken(TEXT("uint4x2"), EHlslToken::Uint4x2);
+				InsertToken(TEXT("uint1x3"), EHlslToken::Uint1x3);
+				InsertToken(TEXT("uint2x3"), EHlslToken::Uint2x3);
+				InsertToken(TEXT("uint3x3"), EHlslToken::Uint3x3);
+				InsertToken(TEXT("uint4x3"), EHlslToken::Uint4x3);
+				InsertToken(TEXT("uint1x4"), EHlslToken::Uint1x4);
+				InsertToken(TEXT("uint2x4"), EHlslToken::Uint2x4);
+				InsertToken(TEXT("uint3x4"), EHlslToken::Uint3x4);
+				InsertToken(TEXT("uint4x4"), EHlslToken::Uint4x4);
+
+				InsertToken(TEXT("half"), EHlslToken::Half);
+				InsertToken(TEXT("half1"), EHlslToken::Half1);
+				InsertToken(TEXT("half2"), EHlslToken::Half2);
+				InsertToken(TEXT("half3"), EHlslToken::Half3);
+				InsertToken(TEXT("half4"), EHlslToken::Half4);
+				InsertToken(TEXT("half1x1"), EHlslToken::Half1x1);
+				InsertToken(TEXT("half2x1"), EHlslToken::Half2x1);
+				InsertToken(TEXT("half3x1"), EHlslToken::Half3x1);
+				InsertToken(TEXT("half4x1"), EHlslToken::Half4x1);
+				InsertToken(TEXT("half1x2"), EHlslToken::Half1x2);
+				InsertToken(TEXT("half2x2"), EHlslToken::Half2x2);
+				InsertToken(TEXT("half3x2"), EHlslToken::Half3x2);
+				InsertToken(TEXT("half4x2"), EHlslToken::Half4x2);
+				InsertToken(TEXT("half1x3"), EHlslToken::Half1x3);
+				InsertToken(TEXT("half2x3"), EHlslToken::Half2x3);
+				InsertToken(TEXT("half3x3"), EHlslToken::Half3x3);
+				InsertToken(TEXT("half4x3"), EHlslToken::Half4x3);
+				InsertToken(TEXT("half1x4"), EHlslToken::Half1x4);
+				InsertToken(TEXT("half2x4"), EHlslToken::Half2x4);
+				InsertToken(TEXT("half3x4"), EHlslToken::Half3x4);
+				InsertToken(TEXT("half4x4"), EHlslToken::Half4x4);
+
+				InsertToken(TEXT("float"), EHlslToken::Float);
+				InsertToken(TEXT("float1"), EHlslToken::Float1);
+				InsertToken(TEXT("float2"), EHlslToken::Float2);
+				InsertToken(TEXT("float3"), EHlslToken::Float3);
+				InsertToken(TEXT("float4"), EHlslToken::Float4);
+				InsertToken(TEXT("float1x1"), EHlslToken::Float1x1);
+				InsertToken(TEXT("float2x1"), EHlslToken::Float2x1);
+				InsertToken(TEXT("float3x1"), EHlslToken::Float3x1);
+				InsertToken(TEXT("float4x1"), EHlslToken::Float4x1);
+				InsertToken(TEXT("float1x2"), EHlslToken::Float1x2);
+				InsertToken(TEXT("float2x2"), EHlslToken::Float2x2);
+				InsertToken(TEXT("float3x2"), EHlslToken::Float3x2);
+				InsertToken(TEXT("float4x2"), EHlslToken::Float4x2);
+				InsertToken(TEXT("float1x3"), EHlslToken::Float1x3);
+				InsertToken(TEXT("float2x3"), EHlslToken::Float2x3);
+				InsertToken(TEXT("float3x3"), EHlslToken::Float3x3);
+				InsertToken(TEXT("float4x3"), EHlslToken::Float4x3);
+				InsertToken(TEXT("float1x4"), EHlslToken::Float1x4);
+				InsertToken(TEXT("float2x4"), EHlslToken::Float2x4);
+				InsertToken(TEXT("float3x4"), EHlslToken::Float3x4);
+				InsertToken(TEXT("float4x4"), EHlslToken::Float4x4);
+
+				InsertToken(TEXT("Texture"), EHlslToken::Texture);
+				InsertToken(TEXT("Texture1D"), EHlslToken::Texture1D);
+				InsertToken(TEXT("Texture1DArray"), EHlslToken::Texture1DArray);
+				InsertToken(TEXT("Texture2D"), EHlslToken::Texture2D);
+				InsertToken(TEXT("Texture2DArray"), EHlslToken::Texture2DArray);
+				InsertToken(TEXT("Texture2DMS"), EHlslToken::Texture2DMS);
+				InsertToken(TEXT("Texture2DMSArray"), EHlslToken::Texture2DMSArray);
+				InsertToken(TEXT("Texture3D"), EHlslToken::Texture3D);
+				InsertToken(TEXT("TextureCube"), EHlslToken::TextureCube);
+				InsertToken(TEXT("TextureCubeArray"), EHlslToken::TextureCubeArray);
+
+				InsertToken(TEXT("Sampler"), EHlslToken::Sampler);
+				InsertToken(TEXT("Sampler1D"), EHlslToken::Sampler1D);
+				InsertToken(TEXT("Sampler2D"), EHlslToken::Sampler2D);
+				InsertToken(TEXT("Sampler3D"), EHlslToken::Sampler3D);
+				InsertToken(TEXT("SamplerCube"), EHlslToken::SamplerCube);
+				InsertToken(TEXT("SamplerState"), EHlslToken::SamplerState);
+				InsertToken(TEXT("SamplerComparisonState"), EHlslToken::SamplerComparisonState);
+
+				InsertToken(TEXT("Buffer"), EHlslToken::Buffer);
+				InsertToken(TEXT("AppendStructuredBuffer"), EHlslToken::AppendStructuredBuffer);
+				InsertToken(TEXT("ByteAddressBuffer"), EHlslToken::ByteAddressBuffer);
+				InsertToken(TEXT("ConsumeStructuredBuffer"), EHlslToken::ConsumeStructuredBuffer);
+				InsertToken(TEXT("RWBuffer"), EHlslToken::RWBuffer);
+				InsertToken(TEXT("RWByteAddressBuffer"), EHlslToken::RWByteAddressBuffer);
+				InsertToken(TEXT("RWStructuredBuffer"), EHlslToken::RWStructuredBuffer);
+				InsertToken(TEXT("RWTexture1D"), EHlslToken::RWTexture1D);
+				InsertToken(TEXT("RWTexture1DArray"), EHlslToken::RWTexture1DArray);
+				InsertToken(TEXT("RWTexture2D"), EHlslToken::RWTexture2D);
+				InsertToken(TEXT("RWTexture2DArray"), EHlslToken::RWTexture2DArray);
+				InsertToken(TEXT("RWTexture3D"), EHlslToken::RWTexture3D);
+				InsertToken(TEXT("StructuredBuffer"), EHlslToken::StructuredBuffer);
+				InsertToken(TEXT("InputPatch"), EHlslToken::InputPatch);
+				InsertToken(TEXT("OutputPatch"), EHlslToken::OutputPatch);
+
+				// Modifiers
+				InsertToken(TEXT("in"), EHlslToken::In);
+				InsertToken(TEXT("out"), EHlslToken::Out);
+				InsertToken(TEXT("inout"), EHlslToken::InOut);
+				InsertToken(TEXT("static"), EHlslToken::Static);
+
+				// Misc
+				InsertToken(TEXT("["), EHlslToken::LeftSquareBracket);
+				InsertToken(TEXT("]"), EHlslToken::RightSquareBracket);
+				InsertToken(TEXT("?"), EHlslToken::Question);
+				InsertToken(TEXT(":"), EHlslToken::Colon);
+				InsertToken(TEXT(","), EHlslToken::Comma);
+				InsertToken(TEXT("."), EHlslToken::Dot);
+				InsertToken(TEXT("struct"), EHlslToken::Struct);
+				InsertToken(TEXT("cbuffer"), EHlslToken::CBuffer);
+				InsertToken(TEXT("groupshared"), EHlslToken::GroupShared);
+				InsertToken(TEXT("nointerpolation"), EHlslToken::NoInterpolation);
+				InsertToken(TEXT("row_major"), EHlslToken::RowMajor);
+			}
+		} GStaticInitializer;
 	}
 
 	struct FTokenizer
@@ -409,13 +416,6 @@ namespace CrossCompiler
 				Line = 1;
 				CurrentLineStart = Current;
 			}
-
-			static bool bInitialized = false;
-			if (!bInitialized)
-			{
-				InitializeTokenTables();
-				bInitialized = true;
-			}
 		}
 
 		bool HasCharsAvailable() const
diff --git a/Engine/Source/Developer/ShaderCompilerCommon/Private/HlslLexer.h b/Engine/Source/Developer/ShaderCompilerCommon/Private/HlslLexer.h
--- a/Engine/Source/Developer/ShaderCompilerCommon/Private/HlslLexer.h
+++ b/Engine/Source/Developer/ShaderCompilerCommon/Private/HlslLexer.h
@@ -66,6 +66,7 @@ namespace CrossCompiler
 		Case,
 		Break,
 		Default,
+		Continue,
 		Goto,
 
 		// Unary
@@ -244,6 +245,7 @@ namespace CrossCompiler
 		UnsignedIntegerConstant,
 		FloatConstant,
 		BoolConstant,
+		StringConstant,	// C-style "string"
 	};
 
 	struct FSourceInfo
diff --git a/Engine/Source/Developer/ShaderCompilerCommon/Private/HlslParser.cpp b/Engine/Source/Developer/ShaderCompilerCommon/Private/HlslParser.cpp
--- a/Engine/Source/Developer/ShaderCompilerCommon/Private/HlslParser.cpp
+++ b/Engine/Source/Developer/ShaderCompilerCommon/Private/HlslParser.cpp
@@ -10,48 +10,80 @@
 
 namespace CrossCompiler
 {
-	typedef EParseResult(*TTryRule)(class FHlslParser& Scanner);
+	EParseResult ParseExpressionStatement(class FHlslParser& Parser, FLinearAllocator* Allocator, AST::FNode** OutStatement);
+	EParseResult ParseStructBody(FHlslScanner& Scanner, FSymbolScope* SymbolScope, FLinearAllocator* Allocator, AST::FTypeSpecifier** OutTypeSpecifier);
+	EParseResult TryParseAttribute(FHlslParser& Parser, FLinearAllocator* Allocator, AST::FAttribute** OutAttribute);
+
+	typedef EParseResult(*TTryRule)(class FHlslParser& Scanner, FLinearAllocator* Allocator, AST::FNode** OutStatement);
 	struct FRulePair
 	{
 		EHlslToken Token;
-		TTryRule OptionalPrefix;
 		TTryRule TryRule;
+		bool bSupportsAttributes;
 
-		FRulePair(EHlslToken InToken, TTryRule InTryRule, TTryRule InOptionalPrefix = nullptr) : Token(InToken), OptionalPrefix(InOptionalPrefix), TryRule(InTryRule) {}
+		FRulePair(EHlslToken InToken, TTryRule InTryRule, bool bInSupportsAttributes = false) : Token(InToken), TryRule(InTryRule), bSupportsAttributes(bInSupportsAttributes) {}
 	};
 	typedef TArray<FRulePair> TRulesArray;
 
 	class FHlslParser
 	{
 	public:
-		FHlslParser();
+		FHlslParser(FLinearAllocator* InAllocator);
 		FHlslScanner Scanner;
 		FSymbolScope GlobalScope;
 		FSymbolScope* CurrentScope;
+		FLinearAllocator* Allocator;
 	};
 
 	TRulesArray RulesTranslationUnit;
 	TRulesArray RulesStatements;
 
-	EParseResult TryRules(const TRulesArray& Rules, FHlslParser& Parser, const FString& RuleNames, bool bErrorIfNoMatch)
+	EParseResult TryRules(const TRulesArray& Rules, FHlslParser& Parser, const FString& RuleNames, bool bErrorIfNoMatch, FLinearAllocator* Allocator, AST::FNode** OutNode)
 	{
 		for (const auto& Rule : Rules)
 		{
 			auto CurrentTokenIndex = Parser.Scanner.GetCurrentTokenIndex();
-			if (Rule.OptionalPrefix)
+			TLinearArray<AST::FAttribute*> Attributes(Allocator);
+			if (Rule.bSupportsAttributes)
 			{
-				if ((*Rule.OptionalPrefix)(Parser) == EParseResult::Error)
+				while (Parser.Scanner.HasMoreTokens())
 				{
-					return EParseResult::Error;
+					const auto* Peek = Parser.Scanner.GetCurrentToken();
+					if (Peek->Token == EHlslToken::LeftSquareBracket)
+					{
+						AST::FAttribute* Attribute = nullptr;
+						auto Result = TryParseAttribute(Parser, Allocator, &Attribute);
+						if (Result == EParseResult::Matched)
+						{
+							Attributes.Add(Attribute);
+							continue;
+						}
+						else if (Result == EParseResult::Error)
+						{
+							return EParseResult::Error;
+						}
+					}
+
+					break;
 				}
 			}
 
 			if (Parser.Scanner.MatchToken(Rule.Token) || Rule.Token == EHlslToken::Invalid)
 			{
-				EParseResult Result = (*Rule.TryRule)(Parser);
-				if (Result == EParseResult::Error || Result == EParseResult::Matched)
+				AST::FNode* Node = nullptr;
+				EParseResult Result = (*Rule.TryRule)(Parser, Allocator, &Node);
+				if (Result == EParseResult::Error)
 				{
-					return Result;
+					return EParseResult::Error;
+				}
+				else if (Result == EParseResult::Matched)
+				{
+					if (Attributes.Num() > 0)
+					{
+						Swap(Node->Attributes, Attributes);
+					}
+					*OutNode = Node;
+					return EParseResult::Matched;
 				}
 			}
 
@@ -67,25 +99,15 @@ namespace CrossCompiler
 		return EParseResult::NotMatched;
 	}
 
-	EParseResult ParseArrayBracketsAndIndex(FHlslScanner& Scanner, FSymbolScope* SymbolScope, bool bNeedsDimension)
+	EParseResult ParseDeclarationArrayBracketsAndIndex(FHlslScanner& Scanner, FSymbolScope* SymbolScope, bool bNeedsDimension, FLinearAllocator* Allocator, AST::FExpression** OutExpression)
 	{
 		if (Scanner.MatchToken(EHlslToken::LeftSquareBracket))
 		{
-			if (bNeedsDimension)
-			{
-				if (ParseExpression(Scanner, SymbolScope) != EParseResult::Matched)
-				{
-					Scanner.SourceError(TEXT("Expected expression!"));
-					return EParseResult::Error;
-				}
-			}
-			else
+			auto ExpressionResult = ParseExpression(Scanner, SymbolScope, false, Allocator, OutExpression);
+			if (ExpressionResult == EParseResult::Error)
 			{
-				if (ParseExpression(Scanner, SymbolScope) == EParseResult::Error)
-				{
-					Scanner.SourceError(TEXT("Expected expression!"));
-					return EParseResult::Error;
-				}
+				Scanner.SourceError(TEXT("Expected expression!"));
+				return EParseResult::Error;
 			}
 
 			if (!Scanner.MatchToken(EHlslToken::RightSquareBracket))
@@ -94,18 +116,28 @@ namespace CrossCompiler
 				return EParseResult::Error;
 			}
 
+			if (ExpressionResult == EParseResult::NotMatched)
+			{
+				if (bNeedsDimension)
+				{
+					Scanner.SourceError(TEXT("Expected array dimension!"));
+					return EParseResult::Error;
+				}
+			}
+
 			return EParseResult::Matched;
 		}
 
 		return EParseResult::NotMatched;
 	}
 
-	EParseResult ParseMultiArrayBracketsAndIndex(FHlslScanner& Scanner, FSymbolScope* SymbolScope, bool bNeedsDimension)
+	EParseResult ParseDeclarationMultiArrayBracketsAndIndex(FHlslScanner& Scanner, FSymbolScope* SymbolScope, bool bNeedsDimension, FLinearAllocator* Allocator, AST::FDeclaration* Declaration)
 	{
 		bool bFoundOne = false;
 		do
 		{
-			auto Result = ParseArrayBracketsAndIndex(Scanner, SymbolScope, bNeedsDimension);
+			AST::FExpression* Dimension = nullptr;
+			auto Result = ParseDeclarationArrayBracketsAndIndex(Scanner, SymbolScope, bNeedsDimension, Allocator, &Dimension);
 			if (Result == EParseResult::Error)
 			{
 				return EParseResult::Error;
@@ -115,73 +147,97 @@ namespace CrossCompiler
 				break;
 			}
 
+			Declaration->ArraySize.Add(Dimension);
+
 			bFoundOne = true;
 		}
 		while (Scanner.HasMoreTokens());
 
-		return bFoundOne ? EParseResult::Matched : EParseResult::NotMatched;
+		if (bFoundOne)
+		{
+			Declaration->bIsArray = true;
+			return EParseResult::Matched;
+		}
+
+		return EParseResult::NotMatched;
 	}
 
-	EParseResult ParseTextureOrBufferSimpleDeclaration(FHlslScanner& Scanner, FSymbolScope* SymbolScope, bool bMultiple)
+	EParseResult ParseTextureOrBufferSimpleDeclaration(FHlslScanner& Scanner, FSymbolScope* SymbolScope, bool bMultiple, FLinearAllocator* Allocator, AST::FDeclaratorList** OutDeclaratorList)
 	{
 		auto OriginalToken = Scanner.GetCurrentTokenIndex();
-		if (ParseGeneralType(Scanner, ETF_SAMPLER_TEXTURE_BUFFER) != EParseResult::Matched)
-		{
-			goto Unmatched;
-		}
-
-		if (Scanner.MatchToken(EHlslToken::Lower))
+		const auto* Token = Scanner.GetCurrentToken();
+		auto* FullType = (*OutDeclaratorList)->Type;
+		if (ParseGeneralType(Scanner, ETF_SAMPLER_TEXTURE_BUFFER, nullptr, Allocator, &FullType->Specifier) == EParseResult::Matched)
 		{
-			auto Result = ParseGeneralType(Scanner, ETF_BUILTIN_NUMERIC | ETF_USER_TYPES, SymbolScope);
-			if (Result != EParseResult::Matched)
-			{
-				Scanner.SourceError(TEXT("Expected type!"));
-				return EParseResult::Error;
-			}
-
-			if (Scanner.MatchToken(EHlslToken::Comma))
+			if (Scanner.MatchToken(EHlslToken::Lower))
 			{
-				if (!Scanner.MatchToken(EHlslToken::UnsignedIntegerConstant))
+				AST::FTypeSpecifier* ElementTypeSpecifier = nullptr;
+				auto Result = ParseGeneralType(Scanner, ETF_BUILTIN_NUMERIC | ETF_USER_TYPES, SymbolScope, Allocator, &ElementTypeSpecifier);
+				if (Result != EParseResult::Matched)
 				{
-					Scanner.SourceError(TEXT("Expected constant!"));
+					Scanner.SourceError(TEXT("Expected type!"));
 					return EParseResult::Error;
 				}
-			}
 
-			if (!Scanner.MatchToken(EHlslToken::Greater))
-			{
-				Scanner.SourceError(TEXT("Expected '>'!"));
-				return EParseResult::Error;
-			}
-		}
+				FullType->Specifier->InnerType = ElementTypeSpecifier->TypeName;
 
-		{
-NextParameter:
-			// Handle 'Sampler2D Sampler'
-			if (ParseGeneralType(Scanner, ETF_SAMPLER_TEXTURE_BUFFER) == EParseResult::Matched)
-			{
-				//...
+				if (Scanner.MatchToken(EHlslToken::Comma))
+				{
+					auto* Integer = Scanner.GetCurrentToken();
+					if (!Scanner.MatchToken(EHlslToken::UnsignedIntegerConstant))
+					{
+						Scanner.SourceError(TEXT("Expected constant!"));
+						return EParseResult::Error;
+					}
+					FullType->Specifier->TextureMSNumSamples = Integer->UnsignedInteger;
+				}
+
+				if (!Scanner.MatchToken(EHlslToken::Greater))
+				{
+					Scanner.SourceError(TEXT("Expected '>'!"));
+					return EParseResult::Error;
+				}
 			}
-			else if (!Scanner.MatchToken(EHlslToken::Identifier))
+			else
 			{
-				Scanner.SourceError(TEXT("Expected Identifier!"));
-				return EParseResult::Error;
+				//TypeSpecifier->InnerName = "float4";
 			}
 
-			if (ParseMultiArrayBracketsAndIndex(Scanner, SymbolScope, true) == EParseResult::Error)
+			do
 			{
-				return EParseResult::Error;
-			}
+				// Handle 'Sampler2D Sampler'
+				AST::FTypeSpecifier* DummyTypeSpecifier = nullptr;
+				const auto* IdentifierToken = Scanner.GetCurrentToken();
+				AST::FDeclaration* Declaration = nullptr;
+				if (ParseGeneralType(Scanner, ETF_SAMPLER_TEXTURE_BUFFER, nullptr, Allocator, &DummyTypeSpecifier) == EParseResult::Matched)
+				{
+					Declaration = new(Allocator) AST::FDeclaration(Allocator, DummyTypeSpecifier->SourceInfo);
+					Declaration->Identifier = Allocator->Strdup(DummyTypeSpecifier->TypeName);
+				}
+				else if (Scanner.MatchToken(EHlslToken::Identifier))
+				{
+					Declaration = new(Allocator) AST::FDeclaration(Allocator, IdentifierToken->SourceInfo);
+					Declaration->Identifier = Allocator->Strdup(IdentifierToken->String);
+				}
+				else
+				{
+					Scanner.SourceError(TEXT("Expected Identifier!"));
+					return EParseResult::Error;
+				}
 
-			if (bMultiple && Scanner.MatchToken(EHlslToken::Comma))
-			{
-				goto NextParameter;
+				if (ParseDeclarationMultiArrayBracketsAndIndex(Scanner, SymbolScope, true, Allocator, Declaration) == EParseResult::Error)
+				{
+					return EParseResult::Error;
+				}
+
+				(*OutDeclaratorList)->Declarations.Add(Declaration);
 			}
+			while (bMultiple && Scanner.MatchToken(EHlslToken::Comma));
 
 			return EParseResult::Matched;
 		}
 
-Unmatched:
+		// Unmatched
 		Scanner.SetCurrentTokenIndex(OriginalToken);
 		return EParseResult::NotMatched;
 	}
@@ -204,11 +260,12 @@ Unmatched:
 		EDF_NOINTERPOLATION				= 0x1000,
 	};
 
-	EParseResult ParseInitializer(FHlslScanner& Scanner, FSymbolScope* SymbolScope, bool bAllowLists)
+	EParseResult ParseInitializer(FHlslScanner& Scanner, FSymbolScope* SymbolScope, bool bAllowLists, FLinearAllocator* Allocator, AST::FExpression** OutList)
 	{
 		if (bAllowLists && Scanner.MatchToken(EHlslToken::LeftBrace))
 		{
-			auto Result = ParseExpressionList(EHlslToken::RightBrace, Scanner, SymbolScope, EHlslToken::LeftBrace);
+			*OutList = new(Allocator) AST::FInitializerListExpression(Allocator, Scanner.GetCurrentToken()->SourceInfo);
+			auto Result = ParseExpressionList(EHlslToken::RightBrace, Scanner, SymbolScope, EHlslToken::LeftBrace, Allocator, *OutList);
 			if (Result != EParseResult::Matched)
 			{
 				Scanner.SourceError(TEXT("Invalid initializer list\n"));
@@ -218,7 +275,8 @@ Unmatched:
 		}
 		else
 		{
-			auto Result = ParseExpression(Scanner, SymbolScope);
+			//@todo-rco?
+			auto Result = ParseExpression(Scanner, SymbolScope, true, Allocator, OutList);
 			if (Result == EParseResult::Error)
 			{
 				Scanner.SourceError(TEXT("Invalid initializer expression\n"));
@@ -230,10 +288,11 @@ Unmatched:
 		return EParseResult::NotMatched;
 	}
 
-	EParseResult ParseDeclarationStorageQualifiers(FHlslScanner& Scanner, int32 Flags, bool& bOutPrimitiveFound)
+	EParseResult ParseDeclarationStorageQualifiers(FHlslScanner& Scanner, int32 Flags, bool& bOutPrimitiveFound, AST::FTypeQualifier* Qualifier)
 	{
 		bOutPrimitiveFound = false;
 		int32 StaticFound = 0;
+		int32 NoInterpolationFound = 0;
 		int32 SharedFound = 0;
 		int32 ConstFound = 0;
 		int32 RowMajorFound = 0;
@@ -265,6 +324,7 @@ Unmatched:
 			if ((Flags & EDF_STATIC) && Scanner.MatchToken(EHlslToken::Static))
 			{
 				++StaticFound;
+				Qualifier->bIsStatic = true;
 				if (StaticFound > 1)
 				{
 					Scanner.SourceError(TEXT("'static' found more than once!\n"));
@@ -273,16 +333,17 @@ Unmatched:
 			}
 			else if ((Flags & EDF_NOINTERPOLATION) && Scanner.MatchToken(EHlslToken::NoInterpolation))
 			{
-				++SharedFound;
-				if (SharedFound > 1)
+				++NoInterpolationFound;
+				if (NoInterpolationFound > 1)
 				{
-					Scanner.SourceError(TEXT("'groupshared' found more than once!\n"));
+					Scanner.SourceError(TEXT("'nointerpolation' found more than once!\n"));
 					return EParseResult::Error;
 				}
 			}
 			else if ((Flags & EDF_SHARED) && Scanner.MatchToken(EHlslToken::GroupShared))
 			{
 				++SharedFound;
+				Qualifier->bShared = true;
 				if (SharedFound > 1)
 				{
 					Scanner.SourceError(TEXT("'groupshared' found more than once!\n"));
@@ -292,6 +353,7 @@ Unmatched:
 			else if ((Flags & EDF_CONST_ROW_MAJOR) && Scanner.MatchToken(EHlslToken::Const))
 			{
 				++ConstFound;
+				Qualifier->bConstant = true;
 				if (ConstFound > 1)
 				{
 					Scanner.SourceError(TEXT("'const' found more than once!\n"));
@@ -301,6 +363,7 @@ Unmatched:
 			else if ((Flags & EDF_CONST_ROW_MAJOR) && Scanner.MatchToken(EHlslToken::RowMajor))
 			{
 				++RowMajorFound;
+				Qualifier->bRowMajor = true;
 				if (RowMajorFound > 1)
 				{
 					Scanner.SourceError(TEXT("'row_major' found more than once!\n"));
@@ -310,6 +373,7 @@ Unmatched:
 			else if ((Flags & EDF_IN_OUT) && Scanner.MatchToken(EHlslToken::In))
 			{
 				++InFound;
+				Qualifier->bIn = true;
 				if (InFound > 1)
 				{
 					Scanner.SourceError(TEXT("'in' found more than once!\n"));
@@ -324,6 +388,7 @@ Unmatched:
 			else if ((Flags & EDF_IN_OUT) && Scanner.MatchToken(EHlslToken::Out))
 			{
 				++OutFound;
+				Qualifier->bOut = true;
 				if (OutFound > 1)
 				{
 					Scanner.SourceError(TEXT("'out' found more than once!\n"));
@@ -338,6 +403,8 @@ Unmatched:
 			else if ((Flags & EDF_IN_OUT) && Scanner.MatchToken(EHlslToken::InOut))
 			{
 				++InOutFound;
+				Qualifier->bIn = true;
+				Qualifier->bOut = true;
 				if (InOutFound > 1)
 				{
 					Scanner.SourceError(TEXT("'inout' found more than once!\n"));
@@ -357,32 +424,34 @@ Unmatched:
 
 		bOutPrimitiveFound = (PrimitiveFound > 0);
 
-		return (ConstFound + RowMajorFound + InFound + OutFound + InOutFound + StaticFound + SharedFound + PrimitiveFound)
+		return (ConstFound + RowMajorFound + InFound + OutFound + InOutFound + StaticFound + SharedFound + PrimitiveFound + NoInterpolationFound)
 			? EParseResult::Matched
 			: EParseResult::NotMatched;
 	}
 
-	EParseResult ParseStructBody(FHlslScanner& Scanner, FSymbolScope* SymbolScope);
-
-	EParseResult ParseGeneralDeclarationNoSemicolon(FHlslScanner& Scanner, FSymbolScope* SymbolScope, int32 Flags)
+	EParseResult ParseGeneralDeclarationNoSemicolon(FHlslScanner& Scanner, FSymbolScope* SymbolScope, int32 Flags, FLinearAllocator* Allocator, AST::FDeclaratorList** OutDeclaratorList)
 	{
 		auto OriginalToken = Scanner.GetCurrentTokenIndex();
 		bool bPrimitiveFound = false;
-		auto Result = ParseDeclarationStorageQualifiers(Scanner, Flags, bPrimitiveFound);
+		auto* FullType = new(Allocator) AST::FFullySpecifiedType(Allocator, Scanner.GetCurrentToken()->SourceInfo);
+		auto Result = ParseDeclarationStorageQualifiers(Scanner, Flags, bPrimitiveFound, &FullType->Qualifier);
 		if (Result == EParseResult::Error)
 		{
 			return EParseResult::Error;
 		}
 		bool bCanBeUnmatched = (Result == EParseResult::NotMatched);
 
+		auto* DeclaratorList = new(Allocator) AST::FDeclaratorList(Allocator, FullType->SourceInfo);
+		DeclaratorList->Type = FullType;
+
 		if (!bPrimitiveFound && (Flags & EDF_PRIMITIVE_DATA_TYPE))
 		{
-			const auto* Token = Scanner.GetCurrentToken();
-			if (Token && Token->Token == EHlslToken::Identifier)
+			const auto* StreamToken = Scanner.GetCurrentToken();
+			if (StreamToken && StreamToken->Token == EHlslToken::Identifier)
 			{
-				if (Token->String == TEXT("PointStream") ||
-					Token->String == TEXT("LineStream") ||
-					Token->String == TEXT("TriangleStream"))
+				if (StreamToken->String == TEXT("PointStream") ||
+					StreamToken->String == TEXT("LineStream") ||
+					StreamToken->String == TEXT("TriangleStream"))
 				{
 					Scanner.Advance();
 					bCanBeUnmatched = false;
@@ -393,7 +462,8 @@ Unmatched:
 						return EParseResult::Error;
 					}
 
-					if (ParseGeneralType(Scanner, ETF_BUILTIN_NUMERIC | ETF_USER_TYPES, SymbolScope) != EParseResult::Matched)
+					AST::FTypeSpecifier* TypeSpecifier = nullptr;
+					if (ParseGeneralType(Scanner, ETF_BUILTIN_NUMERIC | ETF_USER_TYPES, SymbolScope, Allocator, &TypeSpecifier) != EParseResult::Matched)
 					{
 						Scanner.SourceError(TEXT("Expected type!"));
 						return EParseResult::Error;
@@ -405,12 +475,22 @@ Unmatched:
 						return EParseResult::Error;
 					}
 
+					auto* IdentifierToken = Scanner.GetCurrentToken();
 					if (!Scanner.MatchToken(EHlslToken::Identifier))
 					{
 						Scanner.SourceError(TEXT("Expected identifier!"));
 						return EParseResult::Error;
 					}
 
+					TypeSpecifier->InnerType = TypeSpecifier->TypeName;
+					TypeSpecifier->TypeName = Allocator->Strdup(StreamToken->String);
+					FullType->Specifier = TypeSpecifier;
+
+					auto* Declaration = new(Allocator) AST::FDeclaration(Allocator, IdentifierToken->SourceInfo);
+					Declaration->Identifier = Allocator->Strdup(IdentifierToken->String);
+
+					DeclaratorList->Declarations.Add(Declaration);
+					*OutDeclaratorList = DeclaratorList;
 					return EParseResult::Matched;
 				}
 			}
@@ -418,53 +498,81 @@ Unmatched:
 
 		if (Flags & EDF_TEXTURE_SAMPLER_OR_BUFFER)
 		{
-			auto Result = ParseTextureOrBufferSimpleDeclaration(Scanner, SymbolScope, (Flags & EDF_MULTIPLE) == EDF_MULTIPLE);
-			if (Result == EParseResult::Error || Result == EParseResult::Matched)
+			auto Result = ParseTextureOrBufferSimpleDeclaration(Scanner, SymbolScope, (Flags & EDF_MULTIPLE) == EDF_MULTIPLE, Allocator, &DeclaratorList);
+			if (Result == EParseResult::Matched)
 			{
-				return Result;
+				*OutDeclaratorList = DeclaratorList;
+				return EParseResult::Matched;
+			}
+			else if (Result == EParseResult::Error)
+			{
+				return EParseResult::Error;
 			}
 		}
 
+		const bool bAllowInitializerList = (Flags & EDF_INITIALIZER_LIST) == EDF_INITIALIZER_LIST;
+
 		if (Scanner.MatchToken(EHlslToken::Struct))
 		{
-			auto Result = ParseStructBody(Scanner, SymbolScope);
+			auto Result = ParseStructBody(Scanner, SymbolScope, Allocator, &FullType->Specifier);
 			if (Result != EParseResult::Matched)
 			{
 				return EParseResult::Error;
 			}
 
-			if (Scanner.MatchToken(EHlslToken::Identifier))
+			do
 			{
-				//... Instance
-
-				if (Flags & EDF_INITIALIZER)
+				auto* IdentifierToken = Scanner.GetCurrentToken();
+				if (Scanner.MatchToken(EHlslToken::Identifier))
 				{
-					if (Scanner.MatchToken(EHlslToken::Equal))
+					//... Instance
+
+					auto* Declaration = new(Allocator) AST::FDeclaration(Allocator, IdentifierToken->SourceInfo);
+					Declaration->Identifier = Allocator->Strdup(IdentifierToken->String);
+
+					if (ParseDeclarationMultiArrayBracketsAndIndex(Scanner, SymbolScope, false, Allocator, Declaration) == EParseResult::Error)
+					{
+						return EParseResult::Error;
+					}
+
+					if (Flags & EDF_INITIALIZER)
 					{
-						if (ParseInitializer(Scanner, SymbolScope, (Flags & EDF_INITIALIZER_LIST) == EDF_INITIALIZER_LIST) != EParseResult::Matched)
+						if (Scanner.MatchToken(EHlslToken::Equal))
 						{
-							Scanner.SourceError(TEXT("Invalid initializer\n"));
-							return EParseResult::Error;
+							if (ParseInitializer(Scanner, SymbolScope, bAllowInitializerList, Allocator, &Declaration->Initializer) != EParseResult::Matched)
+							{
+								Scanner.SourceError(TEXT("Invalid initializer\n"));
+								return EParseResult::Error;
+							}
 						}
 					}
+				
+					DeclaratorList->Declarations.Add(Declaration);
 				}
 			}
+			while ((Flags & EDF_MULTIPLE) == EDF_MULTIPLE && Scanner.MatchToken(EHlslToken::Comma));
+			*OutDeclaratorList = DeclaratorList;
 		}
 		else
 		{
-			auto Result = ParseGeneralType(Scanner, ETF_BUILTIN_NUMERIC | ETF_USER_TYPES, SymbolScope);
+			auto Result = ParseGeneralType(Scanner, ETF_BUILTIN_NUMERIC | ETF_USER_TYPES, SymbolScope, Allocator, &FullType->Specifier);
 			if (Result == EParseResult::Matched)
 			{
 				bool bMatched = false;
 				do
 				{
+					auto* IdentifierToken = Scanner.GetCurrentToken();
 					if (!Scanner.MatchToken(EHlslToken::Identifier))
 					{
 						Scanner.SetCurrentTokenIndex(OriginalToken);
 						return EParseResult::NotMatched;
 					}
 
-					if (ParseMultiArrayBracketsAndIndex(Scanner, SymbolScope, false) == EParseResult::Error)
+					auto* Declaration = new(Allocator) AST::FDeclaration(Allocator, IdentifierToken->SourceInfo);
+					Declaration->Identifier = Allocator->Strdup(IdentifierToken->String);
+					//AddVar
+
+					if (ParseDeclarationMultiArrayBracketsAndIndex(Scanner, SymbolScope, false, Allocator, Declaration) == EParseResult::Error)
 					{
 						return EParseResult::Error;
 					}
@@ -474,12 +582,14 @@ Unmatched:
 					{
 						if (Scanner.MatchToken(EHlslToken::Colon))
 						{
+							auto* Semantic = Scanner.GetCurrentToken();
 							if (!Scanner.MatchToken(EHlslToken::Identifier))
 							{
 								Scanner.SourceError(TEXT("Expected identifier for semantic!"));
 								return EParseResult::Error;
 							}
 
+							Declaration->Semantic = Allocator->Strdup(Semantic->String);
 							bSemanticFound = true;
 						}
 					}
@@ -488,15 +598,19 @@ Unmatched:
 					{
 						if (Scanner.MatchToken(EHlslToken::Equal))
 						{
-							if (ParseInitializer(Scanner, SymbolScope, (Flags & EDF_INITIALIZER_LIST) == EDF_INITIALIZER_LIST) != EParseResult::Matched)
+							if (ParseInitializer(Scanner, SymbolScope, bAllowInitializerList, Allocator, &Declaration->Initializer) != EParseResult::Matched)
 							{
 								Scanner.SourceError(TEXT("Invalid initializer\n"));
 								return EParseResult::Error;
 							}
 						}
 					}
+
+					DeclaratorList->Declarations.Add(Declaration);
 				}
 				while ((Flags & EDF_MULTIPLE) == EDF_MULTIPLE && Scanner.MatchToken(EHlslToken::Comma));
+
+				*OutDeclaratorList = DeclaratorList;
 			}
 			else if (bCanBeUnmatched && Result == EParseResult::NotMatched)
 			{
@@ -508,11 +622,11 @@ Unmatched:
 		return EParseResult::Matched;
 	}
 
-	EParseResult ParseGeneralDeclaration(FHlslScanner& Scanner, FSymbolScope* SymbolScope, int32 Flags)
+	EParseResult ParseGeneralDeclaration(FHlslScanner& Scanner, FSymbolScope* SymbolScope, FLinearAllocator* Allocator, AST::FDeclaratorList** OutDeclaration, int32 Flags)
 	{
 		auto OriginalToken = Scanner.GetCurrentTokenIndex();
 
-		auto Result = ParseGeneralDeclarationNoSemicolon(Scanner, SymbolScope, Flags);
+		auto Result = ParseGeneralDeclarationNoSemicolon(Scanner, SymbolScope, Flags, Allocator, OutDeclaration);
 		if (Result == EParseResult::NotMatched || Result == EParseResult::Error)
 		{
 			return Result;
@@ -530,10 +644,21 @@ Unmatched:
 		return EParseResult::Matched;
 	}
 
-	EParseResult ParseCBuffer(FHlslParser& Parser)
+	EParseResult ParseCBuffer(FHlslParser& Parser, FLinearAllocator* Allocator, AST::FNode** OutDeclaration)
 	{
-		//FInfo Info;
-		Parser.Scanner.MatchToken(EHlslToken::Identifier);
+		const auto* Token = Parser.Scanner.GetCurrentToken();
+		if (!Token)
+		{
+			Parser.Scanner.SourceError(TEXT("Expected '{'!"));
+			return EParseResult::Error;
+		}
+
+		auto* CBuffer = new(Allocator) AST::FCBufferDeclaration(Allocator, Token->SourceInfo);
+		if (Parser.Scanner.MatchToken(EHlslToken::Identifier))
+		{
+			CBuffer->Name = Allocator->Strdup(Token->String);
+		}
+
 		bool bFoundRightBrace = false;
 		if (Parser.Scanner.MatchToken(EHlslToken::LeftBrace))
 		{
@@ -541,16 +666,17 @@ Unmatched:
 			{
 				if (Parser.Scanner.MatchToken(EHlslToken::RightBrace))
 				{
-
 					if (Parser.Scanner.MatchToken(EHlslToken::Semicolon))
 					{
 						// Optional???
 					}
 
+					*OutDeclaration = CBuffer;
 					return EParseResult::Matched;
 				}
 
-				auto Result = ParseGeneralDeclaration(Parser.Scanner, Parser.CurrentScope, EDF_CONST_ROW_MAJOR | EDF_SEMICOLON | EDF_TEXTURE_SAMPLER_OR_BUFFER);
+				AST::FDeclaratorList* Declaration = nullptr;
+				auto Result = ParseGeneralDeclaration(Parser.Scanner, Parser.CurrentScope, Allocator, &Declaration, EDF_CONST_ROW_MAJOR | EDF_SEMICOLON | EDF_TEXTURE_SAMPLER_OR_BUFFER);
 				if (Result == EParseResult::Error)
 				{
 					return EParseResult::Error;
@@ -559,6 +685,7 @@ Unmatched:
 				{
 					break;
 				}
+				CBuffer->Declarations.Add(Declaration);
 			}
 		}
 
@@ -566,23 +693,25 @@ Unmatched:
 		return EParseResult::Error;
 	}
 
-	EParseResult ParseStructBody(FHlslScanner& Scanner, FSymbolScope* SymbolScope)
+	EParseResult ParseStructBody(FHlslScanner& Scanner, FSymbolScope* SymbolScope, FLinearAllocator* Allocator, AST::FTypeSpecifier** OutTypeSpecifier)
 	{
-		//FInfo Info;
-		//@todo-rco: Test me!
 		const auto* Name = Scanner.GetCurrentToken();
 		if (Name && Scanner.MatchToken(EHlslToken::Identifier))
 		{
 			SymbolScope->Add(Name->String);
 		}
 
+		const TCHAR* Parent = nullptr;
 		if (Scanner.MatchToken(EHlslToken::Colon))
 		{
+			const auto* ParentToken = Scanner.GetCurrentToken();
 			if (!Scanner.MatchToken(EHlslToken::Identifier))
 			{
 				Scanner.SourceError(TEXT("Identifier expected!\n"));
 				return EParseResult::Error;
 			}
+
+			Parent = Allocator->Strdup(ParentToken->String);
 		}
 
 		if (!Scanner.MatchToken(EHlslToken::LeftBrace))
@@ -591,6 +720,10 @@ Unmatched:
 			return EParseResult::Error;
 		}
 
+		auto* Struct = new(Allocator) AST::FStructSpecifier(Allocator, Name->SourceInfo);
+		Struct->ParentName = Allocator->Strdup(Parent);
+		Struct->Name = Allocator->Strdup(Name->String);
+
 		bool bFoundRightBrace = false;
 		while (Scanner.HasMoreTokens())
 		{
@@ -600,7 +733,8 @@ Unmatched:
 				break;
 			}
 
-			auto Result = ParseGeneralDeclaration(Scanner, SymbolScope, EDF_CONST_ROW_MAJOR | EDF_SEMICOLON | EDF_SEMANTIC | EDF_TEXTURE_SAMPLER_OR_BUFFER | EDF_NOINTERPOLATION);
+			AST::FDeclaratorList* Declaration = nullptr;
+			auto Result = ParseGeneralDeclaration(Scanner, SymbolScope, Allocator, &Declaration, EDF_CONST_ROW_MAJOR | EDF_SEMICOLON | EDF_SEMANTIC | EDF_TEXTURE_SAMPLER_OR_BUFFER | EDF_NOINTERPOLATION);
 			if (Result == EParseResult::Error)
 			{
 				return EParseResult::Error;
@@ -609,6 +743,7 @@ Unmatched:
 			{
 				break;
 			}
+			Struct->Declarations.Add(Declaration);
 		}
 
 		if (!bFoundRightBrace)
@@ -617,21 +752,28 @@ Unmatched:
 			return EParseResult::Error;
 		}
 
+		auto* TypeSpecifier = new(Allocator) AST::FTypeSpecifier(Allocator, Struct->SourceInfo);
+		TypeSpecifier->Structure = Struct;
+		*OutTypeSpecifier = TypeSpecifier;
 		return EParseResult::Matched;
 	}
 
-	EParseResult ParseFunctionParameterDeclaration(FHlslParser& Parser)
+	EParseResult ParseFunctionParameterDeclaration(FHlslParser& Parser, FLinearAllocator* Allocator, AST::FFunction* Function)
 	{
 		bool bStrictCheck = false;
 
 		while (Parser.Scanner.HasMoreTokens())
 		{
-			auto Result = ParseGeneralDeclaration(Parser.Scanner, Parser.CurrentScope, EDF_CONST_ROW_MAJOR | EDF_IN_OUT | EDF_TEXTURE_SAMPLER_OR_BUFFER | EDF_INITIALIZER | EDF_SEMANTIC | EDF_PRIMITIVE_DATA_TYPE | EDF_NOINTERPOLATION);
+			AST::FDeclaratorList* Declaration = nullptr;
+			auto Result = ParseGeneralDeclaration(Parser.Scanner, Parser.CurrentScope, Allocator, &Declaration, EDF_CONST_ROW_MAJOR | EDF_IN_OUT | EDF_TEXTURE_SAMPLER_OR_BUFFER | EDF_INITIALIZER | EDF_SEMANTIC | EDF_PRIMITIVE_DATA_TYPE | EDF_NOINTERPOLATION);
 			if (Result == EParseResult::Error)
 			{
 				return EParseResult::Error;
 			}
-			else if (!Parser.Scanner.MatchToken(EHlslToken::Comma))
+			
+			auto* Parameter = AST::FParameterDeclarator::CreateFromDeclaratorList(Declaration, Allocator);
+			Function->Parameters.Add(Parameter);
+			if (!Parser.Scanner.MatchToken(EHlslToken::Comma))
 			{
 				break;
 			}
@@ -645,17 +787,20 @@ Unmatched:
 		return EParseResult::Matched;
 	}
 
-	EParseResult ParseFunctionDeclarator(FHlslParser& Parser)
+	EParseResult ParseFunctionDeclarator(FHlslParser& Parser, FLinearAllocator* Allocator, AST::FFunction** OutFunction)
 	{
 		auto OriginalToken = Parser.Scanner.GetCurrentTokenIndex();
-		auto Result = ParseGeneralType(Parser.Scanner, ETF_BUILTIN_NUMERIC | ETF_SAMPLER_TEXTURE_BUFFER | ETF_USER_TYPES | ETF_VOID, Parser.CurrentScope);
+		AST::FTypeSpecifier* TypeSpecifier = nullptr;
+		auto Result = ParseGeneralType(Parser.Scanner, ETF_BUILTIN_NUMERIC | ETF_SAMPLER_TEXTURE_BUFFER | ETF_USER_TYPES | ETF_VOID, Parser.CurrentScope, Allocator, &TypeSpecifier);
 		if (Result == EParseResult::NotMatched)
 		{
 			Parser.Scanner.SetCurrentTokenIndex(OriginalToken);
 			return EParseResult::NotMatched;
 		}
+
 		check(Result == EParseResult::Matched);
 
+		auto* Identifier = Parser.Scanner.GetCurrentToken();
 		if (!Parser.Scanner.MatchToken(EHlslToken::Identifier))
 		{
 			// This could be an error... But we should allow testing for a global variable before any rash decisions
@@ -670,17 +815,22 @@ Unmatched:
 			return EParseResult::NotMatched;
 		}
 
+		auto* Function = new(Allocator) AST::FFunction(Allocator, Identifier->SourceInfo);
+		Function->Identifier = Allocator->Strdup(Identifier->String);
+		Function->ReturnType = new(Allocator) AST::FFullySpecifiedType(Allocator, TypeSpecifier->SourceInfo);
+		Function->ReturnType->Specifier = TypeSpecifier;
+
 		if (Parser.Scanner.MatchToken(EHlslToken::Void))
 		{
-			//...
+			// Nothing to do here...
 		}
 		else if (Parser.Scanner.MatchToken(EHlslToken::RightParenthesis))
 		{
-			return EParseResult::Matched;
+			goto Done;
 		}
 		else
 		{
-			Result = ParseFunctionParameterDeclaration(Parser);
+			Result = ParseFunctionParameterDeclaration(Parser, Allocator, Function);
 			if (Result == EParseResult::Error)
 			{
 				return EParseResult::Error;
@@ -693,10 +843,13 @@ Unmatched:
 			return EParseResult::Error;
 		}
 
+Done:
+		*OutFunction = Function;
+
 		return EParseResult::Matched;
 	}
 
-	EParseResult ParseStatement(FHlslParser& Parser)
+	EParseResult ParseStatement(FHlslParser& Parser, FLinearAllocator* Allocator, AST::FNode** OutStatement)
 	{
 		const auto* Token = Parser.Scanner.PeekToken();
 		if (Token && Token->Token == EHlslToken::RightBrace)
@@ -705,19 +858,22 @@ Unmatched:
 		}
 
 		static FString Statement(TEXT("Statement"));
-		return TryRules(RulesStatements, Parser, Statement, false);
+		return TryRules(RulesStatements, Parser, Statement, false, Allocator, OutStatement);
 	}
 
-	EParseResult ParseStatementBlock(FHlslParser& Parser)
+	EParseResult ParseStatementBlock(FHlslParser& Parser, FLinearAllocator* Allocator, AST::FNode** OutStatement)
 	{
-		FCreateSymbolScope SymbolScope(&Parser.CurrentScope);
+		FCreateSymbolScope SymbolScope(Allocator, &Parser.CurrentScope);
+		auto* Block = new(Allocator) AST::FCompoundStatement(Allocator, Parser.Scanner.GetCurrentToken()->SourceInfo);
 		while (Parser.Scanner.HasMoreTokens())
 		{
-			auto Result = ParseStatement(Parser);
+			AST::FNode* Statement = nullptr;
+			auto Result = ParseStatement(Parser, Allocator, &Statement);
 			if (Result == EParseResult::NotMatched)
 			{
 				if (Parser.Scanner.MatchToken(EHlslToken::RightBrace))
 				{
+					*OutStatement = Block;
 					return EParseResult::Matched;
 				}
 				else
@@ -730,15 +886,23 @@ Unmatched:
 			{
 				break;
 			}
+
+			if (Statement)
+			{
+				Block->Statements.Add(Statement);
+			}
 		}
 
 		Parser.Scanner.SourceError(TEXT("'}' expected!"));
 		return EParseResult::Error;
 	}
 
-	EParseResult ParseFunctionDeclaration(FHlslParser& Parser)
+	EParseResult ParseFunctionDeclaration(FHlslParser& Parser, FLinearAllocator* Allocator, AST::FNode** OutFunction)
 	{
-		EParseResult Result = ParseFunctionDeclarator(Parser);
+		const auto* CurrentToken = Parser.Scanner.GetCurrentToken();
+
+		AST::FFunction* Function = nullptr;
+		EParseResult Result = ParseFunctionDeclarator(Parser, Allocator, &Function);
 		if (Result == EParseResult::NotMatched || Result == EParseResult::Error)
 		{
 			return Result;
@@ -746,6 +910,7 @@ Unmatched:
 
 		if (Parser.Scanner.MatchToken(EHlslToken::Semicolon))
 		{
+check(0);
 			// Forward declare
 			return EParseResult::Matched;
 		}
@@ -767,30 +932,47 @@ Unmatched:
 				return EParseResult::Error;
 			}
 
-			return ParseStatementBlock(Parser);
+			auto* FunctionDefinition = new(Allocator) AST::FFunctionDefinition(Allocator, CurrentToken->SourceInfo);
+			AST::FNode* Body = nullptr;
+			Result = ParseStatementBlock(Parser, Allocator, &Body);
+			if (Result == EParseResult::Matched)
+			{
+				FunctionDefinition->Body = (AST::FCompoundStatement*)Body;
+				FunctionDefinition->Prototype = Function;
+				*OutFunction = FunctionDefinition;
+			}
 		}
 
 		return Result;
 	}
 
-	EParseResult ParseLocalDeclaration(FHlslParser& Parser)
+	EParseResult ParseLocalDeclaration(FHlslParser& Parser, FLinearAllocator* Allocator, AST::FNode** OutDeclaration)
 	{
-		return ParseGeneralDeclaration(Parser.Scanner, Parser.CurrentScope, EDF_CONST_ROW_MAJOR | EDF_INITIALIZER | EDF_INITIALIZER_LIST | EDF_SEMICOLON | EDF_MULTIPLE);
+		AST::FDeclaratorList* List = nullptr;
+		auto Result = ParseGeneralDeclaration(Parser.Scanner, Parser.CurrentScope, Allocator, &List, EDF_CONST_ROW_MAJOR | EDF_INITIALIZER | EDF_INITIALIZER_LIST | EDF_SEMICOLON | EDF_MULTIPLE);
+		*OutDeclaration = List;
+		return Result;
 	}
 
-	EParseResult ParseGlobalVariableDeclaration(FHlslParser& Parser)
+	EParseResult ParseGlobalVariableDeclaration(FHlslParser& Parser, FLinearAllocator* Allocator, AST::FNode** OutDeclaration)
 	{
-		return ParseGeneralDeclaration(Parser.Scanner, Parser.CurrentScope, EDF_CONST_ROW_MAJOR | EDF_STATIC | EDF_SHARED | EDF_TEXTURE_SAMPLER_OR_BUFFER | EDF_INITIALIZER | EDF_INITIALIZER_LIST | EDF_SEMICOLON | EDF_MULTIPLE);
+		AST::FDeclaratorList* List = nullptr;
+		auto Result = ParseGeneralDeclaration(Parser.Scanner, Parser.CurrentScope, Allocator, &List, EDF_CONST_ROW_MAJOR | EDF_STATIC | EDF_SHARED | EDF_TEXTURE_SAMPLER_OR_BUFFER | EDF_INITIALIZER | EDF_INITIALIZER_LIST | EDF_SEMICOLON | EDF_MULTIPLE);
+		*OutDeclaration = List;
+		return Result;
 	}
 
-	EParseResult ParseReturnStatement(FHlslParser& Parser)
+	EParseResult ParseReturnStatement(FHlslParser& Parser, FLinearAllocator* Allocator, AST::FNode** OutStatement)
 	{
+		auto* Statement = new(Allocator) AST::FJumpStatement(Allocator, AST::EJumpType::Return, Parser.Scanner.GetCurrentToken()->SourceInfo);
+
 		if (Parser.Scanner.MatchToken(EHlslToken::Semicolon))
 		{
+			*OutStatement = Statement;
 			return EParseResult::Matched;
 		}
 
-		if (ParseExpression(Parser.Scanner, Parser.CurrentScope) != EParseResult::Matched)
+		if (ParseExpression(Parser.Scanner, Parser.CurrentScope, true, Allocator, &Statement->OptionalExpression) != EParseResult::Matched)
 		{
 			Parser.Scanner.SourceError(TEXT("Expression expected"));
 			return EParseResult::Error;
@@ -802,14 +984,24 @@ Unmatched:
 			return EParseResult::Error;
 		}
 
+		*OutStatement = Statement;
 		return EParseResult::Matched;
 	}
 
-	EParseResult ParseDoStatement(FHlslParser& Parser)
+	EParseResult ParseDoStatement(FHlslParser& Parser, FLinearAllocator* Allocator, AST::FNode** OutStatement)
 	{
-		if (ParseStatement(Parser) != EParseResult::Matched)
+		FCreateSymbolScope SymbolScope(Allocator, &Parser.CurrentScope);
+		const auto* Token = Parser.Scanner.GetCurrentToken();
+		AST::FNode* Body = nullptr;
+		auto Result = ParseStatement(Parser, Allocator, &Body);
+		if (Result != EParseResult::Matched)
 		{
-			Parser.Scanner.SourceError(TEXT("Statement expected"));
+			return EParseResult::Error;
+		}
+
+		if (!Parser.Scanner.MatchToken(EHlslToken::While))
+		{
+			Parser.Scanner.SourceError(TEXT("'while' expected"));
 			return EParseResult::Error;
 		}
 
@@ -819,7 +1011,8 @@ Unmatched:
 			return EParseResult::Error;
 		}
 
-		if (ParseExpression(Parser.Scanner, Parser.CurrentScope) != EParseResult::Matched)
+		AST::FExpression* ConditionExpression = nullptr;
+		if (ParseExpression(Parser.Scanner, Parser.CurrentScope, true, Allocator, &ConditionExpression) != EParseResult::Matched)
 		{
 			Parser.Scanner.SourceError(TEXT("Expression expected"));
 			return EParseResult::Error;
@@ -831,26 +1024,31 @@ Unmatched:
 			return EParseResult::Error;
 		}
 
-		if (ParseStatement(Parser) != EParseResult::Matched)
+		if (!Parser.Scanner.MatchToken(EHlslToken::Semicolon))
 		{
-			Parser.Scanner.SourceError(TEXT("Statement expected"));
+			Parser.Scanner.SourceError(TEXT("';' expected"));
 			return EParseResult::Error;
 		}
 
+		auto* DoWhile = new(Allocator) AST::FIterationStatement(Allocator, Token->SourceInfo, AST::EIterationType::DoWhile);
+		DoWhile->Condition = ConditionExpression;
+		DoWhile->Body = Body;
+		*OutStatement = DoWhile;
 		return EParseResult::Matched;
 	}
 
-	EParseResult ParseWhileStatement(FHlslParser& Parser)
+	EParseResult ParseWhileStatement(FHlslParser& Parser, FLinearAllocator* Allocator, AST::FNode** OutStatement)
 	{
-		FCreateSymbolScope SymbolScope(&Parser.CurrentScope);
-
+		FCreateSymbolScope SymbolScope(Allocator, &Parser.CurrentScope);
+		const auto* Token = Parser.Scanner.GetCurrentToken();
 		if (!Parser.Scanner.MatchToken(EHlslToken::LeftParenthesis))
 		{
 			Parser.Scanner.SourceError(TEXT("'(' expected"));
 			return EParseResult::Error;
 		}
 
-		if (ParseExpression(Parser.Scanner, Parser.CurrentScope) != EParseResult::Matched)
+		AST::FExpression* ConditionExpression = nullptr;
+		if (ParseExpression(Parser.Scanner, Parser.CurrentScope, true, Allocator, &ConditionExpression) != EParseResult::Matched)
 		{
 			Parser.Scanner.SourceError(TEXT("Expression expected"));
 			return EParseResult::Error;
@@ -862,48 +1060,56 @@ Unmatched:
 			return EParseResult::Error;
 		}
 
-		if (ParseStatement(Parser) != EParseResult::Matched)
+		AST::FNode* Body = nullptr;
+		auto Result = ParseStatement(Parser, Allocator, &Body);
+		if (Result != EParseResult::Matched)
 		{
-			Parser.Scanner.SourceError(TEXT("Statement expected"));
 			return EParseResult::Error;
 		}
 
+		auto* While = new(Allocator) AST::FIterationStatement(Allocator, Token->SourceInfo, AST::EIterationType::While);
+		While->Condition = ConditionExpression;
+		While->Body = Body;
+		*OutStatement = While;
 		return EParseResult::Matched;
 	}
 
-	EParseResult ParseForStatement(FHlslParser& Parser)
+	EParseResult ParseForStatement(FHlslParser& Parser, FLinearAllocator* Allocator, AST::FNode** OutStatement)
 	{
-		FCreateSymbolScope SymbolScope(&Parser.CurrentScope);
-
+		FCreateSymbolScope SymbolScope(Allocator, &Parser.CurrentScope);
+		const auto* Token = Parser.Scanner.GetCurrentToken();
 		if (!Parser.Scanner.MatchToken(EHlslToken::LeftParenthesis))
 		{
 			Parser.Scanner.SourceError(TEXT("Expected '('!\n"));
 			return EParseResult::Error;
 		}
 
-		auto Result = ParseGeneralDeclaration(Parser.Scanner, Parser.CurrentScope, EDF_CONST_ROW_MAJOR | EDF_INITIALIZER);
-		if (Result == EParseResult::Error)
+		AST::FNode* InitExpression = nullptr;
+		if (Parser.Scanner.MatchToken(EHlslToken::Semicolon))
 		{
-			Parser.Scanner.SourceError(TEXT("Expected expression or declaration!\n"));
-			return EParseResult::Error;
+			// Do nothing...
 		}
-		else if (Result == EParseResult::NotMatched)
+		else
 		{
-			Result = ParseExpression(Parser.Scanner, Parser.CurrentScope);
+			auto Result = ParseLocalDeclaration(Parser, Allocator, &InitExpression);
 			if (Result == EParseResult::Error)
 			{
 				Parser.Scanner.SourceError(TEXT("Expected expression or declaration!\n"));
 				return EParseResult::Error;
 			}
+			else if (Result == EParseResult::NotMatched)
+			{
+				Result = ParseExpressionStatement(Parser, Allocator, &InitExpression);
+				if (Result == EParseResult::Error)
+				{
+					Parser.Scanner.SourceError(TEXT("Expected expression or declaration!\n"));
+					return EParseResult::Error;
+				}
+			}
 		}
 
-		if (!Parser.Scanner.MatchToken(EHlslToken::Semicolon))
-		{
-			Parser.Scanner.SourceError(TEXT("Expected ';'!\n"));
-			return EParseResult::Error;
-		}
-
-		Result = ParseExpression(Parser.Scanner, Parser.CurrentScope);
+		AST::FExpression* ConditionExpression = nullptr;
+		auto Result = ParseExpression(Parser.Scanner, Parser.CurrentScope, true, Allocator, &ConditionExpression);
 		if (Result == EParseResult::Error)
 		{
 			Parser.Scanner.SourceError(TEXT("Expected expression or declaration!\n"));
@@ -916,7 +1122,8 @@ Unmatched:
 			return EParseResult::Error;
 		}
 
-		Result = ParseExpression(Parser.Scanner, Parser.CurrentScope);
+		AST::FExpression* RestExpression = nullptr;
+		Result = ParseExpression(Parser.Scanner, Parser.CurrentScope, true, Allocator, &RestExpression);
 		if (Result == EParseResult::Error)
 		{
 			Parser.Scanner.SourceError(TEXT("Expected expression or declaration!\n"));
@@ -929,12 +1136,27 @@ Unmatched:
 			return EParseResult::Error;
 		}
 
-		return ParseStatement(Parser);
+		AST::FNode* Body = nullptr;
+		Result = ParseStatement(Parser, Allocator, &Body);
+		if (Result != EParseResult::Matched)
+		{
+			return EParseResult::Error;
+		}
+
+		auto* For = new(Allocator) AST::FIterationStatement(Allocator, Token->SourceInfo, AST::EIterationType::For);
+		For->InitStatement = InitExpression;
+		For->Condition = ConditionExpression;
+		For->RestExpression = RestExpression;
+		For->Body = Body;
+		*OutStatement = For;
+		return EParseResult::Matched;
 	}
 
-	EParseResult ParseIfStatement(FHlslParser& Parser)
+	EParseResult ParseIfStatement(FHlslParser& Parser, FLinearAllocator* Allocator, AST::FNode** OutStatement)
 	{
-		FCreateSymbolScope SymbolScope(&Parser.CurrentScope);
+		FCreateSymbolScope SymbolScope(Allocator, &Parser.CurrentScope);
+
+		auto* Statement = new(Allocator) AST::FSelectionStatement(Allocator, Parser.Scanner.GetCurrentToken()->SourceInfo);
 
 		if (!Parser.Scanner.MatchToken(EHlslToken::LeftParenthesis))
 		{
@@ -942,7 +1164,7 @@ Unmatched:
 			return EParseResult::Error;
 		}
 
-		if (ParseExpression(Parser.Scanner, Parser.CurrentScope) != EParseResult::Matched)
+		if (ParseExpression(Parser.Scanner, Parser.CurrentScope, true, Allocator, &Statement->Condition) != EParseResult::Matched)
 		{
 			Parser.Scanner.SourceError(TEXT("Expression expected"));
 			return EParseResult::Error;
@@ -954,7 +1176,7 @@ Unmatched:
 			return EParseResult::Error;
 		}
 
-		if (ParseStatement(Parser) != EParseResult::Matched)
+		if (ParseStatement(Parser, Allocator, &Statement->ThenStatement) != EParseResult::Matched)
 		{
 			Parser.Scanner.SourceError(TEXT("Statement expected"));
 			return EParseResult::Error;
@@ -962,29 +1184,80 @@ Unmatched:
 
 		if (Parser.Scanner.MatchToken(EHlslToken::Else))
 		{
-			if (ParseStatement(Parser) != EParseResult::Matched)
+			if (ParseStatement(Parser, Allocator, &Statement->ElseStatement) != EParseResult::Matched)
 			{
 				Parser.Scanner.SourceError(TEXT("Statement expected"));
 				return EParseResult::Error;
 			}
 		}
 
+		*OutStatement = Statement;
 		return EParseResult::Matched;
 	}
 
-	EParseResult TryParseAttribute(FHlslParser& Parser)
+	EParseResult ParseAttributeArgList(FHlslScanner& Scanner, FSymbolScope* SymbolScope, FLinearAllocator* Allocator, AST::FAttribute* OutAttribute)
+	{
+		while (Scanner.HasMoreTokens())
+		{
+			const auto* Token = Scanner.PeekToken();
+			if (Scanner.MatchToken(EHlslToken::RightParenthesis))
+			{
+				return EParseResult::Matched;
+			}
+
+			bool bMultiple = false;
+			do
+			{
+				bMultiple = false;
+				Token = Scanner.PeekToken();
+				if (Scanner.MatchToken(EHlslToken::StringConstant))
+				{
+					auto* Arg = new(Allocator) AST::FAttributeArgument(Allocator, Token->SourceInfo);
+					Arg->StringArgument = Allocator->Strdup(Token->String);
+					OutAttribute->Arguments.Add(Arg);
+				}
+				else
+				{
+					AST::FExpression* Expression = nullptr;
+					EParseResult Result = ParseExpression(Scanner, SymbolScope, false, Allocator, &Expression);
+					if (Result != EParseResult::Matched)
+					{
+						Scanner.SourceError(TEXT("Incorrect attribute expression!\n"));
+						return EParseResult::Error;
+					}
+					auto* Arg = new(Allocator) AST::FAttributeArgument(Allocator, Token->SourceInfo);
+					Arg->ExpressionArgument = Expression;
+					OutAttribute->Arguments.Add(Arg);
+				}
+
+				if (Scanner.MatchToken(EHlslToken::Comma))
+				{
+					bMultiple = true;
+				}
+			}
+			while (bMultiple);
+		}
+
+		return EParseResult::Error;
+	}
+
+	EParseResult TryParseAttribute(FHlslParser& Parser, FLinearAllocator* Allocator, AST::FAttribute** OutAttribute)
 	{
+		const auto* Token = Parser.Scanner.GetCurrentToken();
 		if (Parser.Scanner.MatchToken(EHlslToken::LeftSquareBracket))
 		{
+			auto* Identifier = Parser.Scanner.GetCurrentToken();
 			if (!Parser.Scanner.MatchToken(EHlslToken::Identifier))
 			{
 				Parser.Scanner.SourceError(TEXT("Incorrect attribute\n"));
 				return EParseResult::Error;
 			}
 
+			auto* Attribute = new(Allocator) AST::FAttribute(Allocator, Token->SourceInfo, Allocator->Strdup(Identifier->String));
+
 			if (Parser.Scanner.MatchToken(EHlslToken::LeftParenthesis))
 			{
-				auto Result = ParseExpressionList(EHlslToken::RightParenthesis, Parser.Scanner, Parser.CurrentScope);
+				auto Result = ParseAttributeArgList(Parser.Scanner, Parser.CurrentScope, Allocator, Attribute);
 				if (Result != EParseResult::Matched)
 				{
 					Parser.Scanner.SourceError(TEXT("Incorrect attribute! Expected ')'.\n"));
@@ -998,121 +1271,187 @@ Unmatched:
 				return EParseResult::Error;
 			}
 
+			*OutAttribute = Attribute;
 			return EParseResult::Matched;
 		}
 
 		return EParseResult::NotMatched;
 	}
 
-	EParseResult ParseSwitchStatement(FHlslParser& Parser)
+	EParseResult ParseSwitchBody(FHlslParser& Parser, FLinearAllocator* Allocator, AST::FSwitchBody** OutBody)
 	{
-		if (!Parser.Scanner.MatchToken(EHlslToken::LeftParenthesis))
-		{
-			Parser.Scanner.SourceError(TEXT("'(' expected"));
-			return EParseResult::Error;
-		}
+		const auto* Token = Parser.Scanner.GetCurrentToken();
 
-		if (ParseExpression(Parser.Scanner, Parser.CurrentScope) != EParseResult::Matched)
+		if (!Parser.Scanner.MatchToken(EHlslToken::LeftBrace))
 		{
-			Parser.Scanner.SourceError(TEXT("Expression expected"));
+			Parser.Scanner.SourceError(TEXT("'{' expected"));
 			return EParseResult::Error;
 		}
 
-		if (!Parser.Scanner.MatchToken(EHlslToken::RightParenthesis))
-		{
-			Parser.Scanner.SourceError(TEXT("')' expected"));
-			return EParseResult::Error;
-		}
+		auto* Body = new(Allocator) AST::FSwitchBody(Allocator, Token->SourceInfo);
 
-		if (!Parser.Scanner.MatchToken(EHlslToken::LeftBrace))
+		// Empty switch
+		if (Parser.Scanner.MatchToken(EHlslToken::RightBrace))
 		{
-			Parser.Scanner.SourceError(TEXT("'{' expected"));
-			return EParseResult::Error;
+			*OutBody = Body;
+			return EParseResult::Matched;
 		}
 
+		Body->CaseList = new(Allocator) AST::FCaseStatementList(Allocator, Token->SourceInfo);
+
 		bool bDefaultFound = false;
 		while (Parser.Scanner.HasMoreTokens())
 		{
-			if (Parser.Scanner.MatchToken(EHlslToken::Default))
+			auto* Token = Parser.Scanner.GetCurrentToken();
+			if (Parser.Scanner.MatchToken(EHlslToken::RightBrace))
 			{
-				if (bDefaultFound)
-				{
-					Parser.Scanner.SourceError(TEXT("'default' found twice on switch() statement!"));
-					return EParseResult::Error;
-				}
+				break;
+			}
 
-				if (!Parser.Scanner.MatchToken(EHlslToken::Colon))
-				{
-					Parser.Scanner.SourceError(TEXT("':' expected"));
-					return EParseResult::Error;
-				}
+			auto* Labels = new(Allocator) AST::FCaseLabelList(Allocator, Token->SourceInfo);
+			auto* CaseStatement = new(Allocator) AST::FCaseStatement(Allocator, Token->SourceInfo, Labels);
 
-				bDefaultFound = true;
-			}
-			else if (Parser.Scanner.MatchToken(EHlslToken::Case))
+			// Case labels
+			bool bLabelFound = false;
+			do
 			{
-				if (ParseExpression(Parser.Scanner, Parser.CurrentScope) != EParseResult::Matched)
+				bLabelFound = false;
+				AST::FCaseLabel* Label = nullptr;
+				Token = Parser.Scanner.GetCurrentToken();
+				if (Parser.Scanner.MatchToken(EHlslToken::Default))
 				{
-					Parser.Scanner.SourceError(TEXT("Expression expected on case label!"));
-					return EParseResult::Error;
-				}
+					if (bDefaultFound)
+					{
+						Parser.Scanner.SourceError(TEXT("'default' found twice on switch() statement!"));
+						return EParseResult::Error;
+					}
+
+					if (!Parser.Scanner.MatchToken(EHlslToken::Colon))
+					{
+						Parser.Scanner.SourceError(TEXT("':' expected"));
+						return EParseResult::Error;
+					}
 
-				if (!Parser.Scanner.MatchToken(EHlslToken::Colon))
+					Label = new(Allocator) AST::FCaseLabel(Allocator, Token->SourceInfo, nullptr);
+					bDefaultFound = true;
+					bLabelFound = true;
+				}
+				else if (Parser.Scanner.MatchToken(EHlslToken::Case))
 				{
-					Parser.Scanner.SourceError(TEXT("':' expected"));
-					return EParseResult::Error;
+					AST::FExpression* CaseExpression = nullptr;
+					if (ParseExpression(Parser.Scanner, Parser.CurrentScope, true, Allocator, &CaseExpression) != EParseResult::Matched)
+					{
+						Parser.Scanner.SourceError(TEXT("Expression expected on case label!"));
+						return EParseResult::Error;
+					}
+
+					if (!Parser.Scanner.MatchToken(EHlslToken::Colon))
+					{
+						Parser.Scanner.SourceError(TEXT("':' expected"));
+						return EParseResult::Error;
+					}
+
+					Label = new(Allocator) AST::FCaseLabel(Allocator, Token->SourceInfo, CaseExpression);
+					bLabelFound = true;
 				}
-			}
 
-			auto* Peek = Parser.Scanner.PeekToken();
-			if (!Peek)
-			{
-				break;
-			}
-			else if (Parser.Scanner.MatchToken(EHlslToken::RightBrace))
-			{
-				return EParseResult::Matched;
-			}
-			else if (Peek->Token == EHlslToken::Case || Peek->Token == EHlslToken::Default)
-			{
-				continue;
-			}
-			else
-			{
-				auto Result = ParseStatement(Parser);
-				if (Result == EParseResult::Error)
+				if (Label)
 				{
-					return EParseResult::Error;
+					CaseStatement->Labels->Labels.Add(Label);
 				}
 			}
-/*
-			// Process statements until hitting break, case, default or '}'
+			while (bLabelFound);
+
+			// Statements
+			Token = Parser.Scanner.GetCurrentToken();
+			bool bMatchedOnce = false;
 			while (Parser.Scanner.HasMoreTokens())
 			{
-				if (Parser.Scanner.MatchToken(EHlslToken::RightBrace))
+				auto* Peek = Parser.Scanner.PeekToken();
+				if (!Peek)
 				{
-					return EParseResult::Matched;
+					break;
 				}
-
-				if (Parser.Scanner.MatchToken(EHlslToken::Break))
+				else if (Peek->Token == EHlslToken::RightBrace)
+				{
+					// End of switch
+					break;
+				}
+				else if (Peek->Token == EHlslToken::Case || Peek->Token == EHlslToken::Default)
 				{
+					// Next CaseStatement
 					break;
 				}
+				else
+				{
+					AST::FNode* Statement = nullptr;
+					auto Result = ParseStatement(Parser, Allocator, &Statement);
+					if (Result == EParseResult::Error)
+					{
+						return EParseResult::Error;
+					}
+					else if (Result == EParseResult::NotMatched)
+					{
+						Parser.Scanner.SourceError(TEXT("Internal Error parsing statment inside case list"));
+						return EParseResult::Error;
+					}
+					else
+					{
+						CaseStatement->Statements.Add(Statement);
+					}
+				}
 			}
-			check(0);
-*/
+
+			Body->CaseList->Cases.Add(CaseStatement);
 		}
 
-		return EParseResult::Error;
+		*OutBody = Body;
+		return EParseResult::Matched;
 	}
 
-	EParseResult ParseExpressionStatement(FHlslParser& Parser)
+	EParseResult ParseSwitchStatement(FHlslParser& Parser, FLinearAllocator* Allocator, AST::FNode** OutStatement)
+	{
+		const auto* Token = Parser.Scanner.GetCurrentToken();
+		if (!Parser.Scanner.MatchToken(EHlslToken::LeftParenthesis))
+		{
+			Parser.Scanner.SourceError(TEXT("'(' expected"));
+			return EParseResult::Error;
+		}
+
+		AST::FExpression* Condition = nullptr;
+		if (ParseExpression(Parser.Scanner, Parser.CurrentScope, false, Allocator, &Condition) != EParseResult::Matched)
+		{
+			Parser.Scanner.SourceError(TEXT("Expression expected"));
+			return EParseResult::Error;
+		}
+
+		if (!Parser.Scanner.MatchToken(EHlslToken::RightParenthesis))
+		{
+			Parser.Scanner.SourceError(TEXT("')' expected"));
+			return EParseResult::Error;
+		}
+
+		AST::FSwitchBody* Body = nullptr;
+		if (ParseSwitchBody(Parser, Allocator, &Body) != EParseResult::Matched)
+		{
+			return EParseResult::Error;
+		}
+
+		auto* Switch = new(Allocator) AST::FSwitchStatement(Allocator, Token->SourceInfo, Condition, Body);
+		*OutStatement = Switch;
+
+		return EParseResult::Matched;
+	}
+
+	EParseResult ParseExpressionStatement(FHlslParser& Parser, FLinearAllocator* Allocator, AST::FNode** OutStatement)
 	{
 		auto OriginalToken = Parser.Scanner.GetCurrentTokenIndex();
-		if (ParseExpression(Parser.Scanner, Parser.CurrentScope) == EParseResult::Matched)
+		auto* Statement = new(Allocator) AST::FExpressionStatement(Allocator, nullptr, Parser.Scanner.GetCurrentToken()->SourceInfo);
+		if (ParseExpression(Parser.Scanner, Parser.CurrentScope, true, Allocator, &Statement->Expression) == EParseResult::Matched)
 		{
 			if (Parser.Scanner.MatchToken(EHlslToken::Semicolon))
 			{
+				*OutStatement = Statement;
 				return EParseResult::Matched;
 			}
 		}
@@ -1121,62 +1460,99 @@ Unmatched:
 		return EParseResult::NotMatched;
 	}
 
-	EParseResult ParseEmptyStatement(FHlslParser& Parser)
+	EParseResult ParseEmptyStatement(FHlslParser& Parser, FLinearAllocator* Allocator, AST::FNode** OutStatement)
 	{
+		check(*OutStatement == nullptr);
+
 		// Nothing to do here...
 		return EParseResult::Matched;
 	}
 
-	void InitRules()
+	EParseResult ParseBreakStatement(FHlslParser& Parser, FLinearAllocator* Allocator, AST::FNode** OutStatement)
 	{
-		static bool bInit = false;
-		if (!bInit)
-		{
-			bInit = true;
-
-			RulesTranslationUnit.Add(FRulePair(EHlslToken::CBuffer, ParseCBuffer));
-			RulesTranslationUnit.Add(FRulePair(EHlslToken::Invalid, ParseFunctionDeclaration, TryParseAttribute));
-			RulesTranslationUnit.Add(FRulePair(EHlslToken::Invalid, ParseGlobalVariableDeclaration));
-
-			RulesStatements.Add(FRulePair(EHlslToken::LeftBrace, ParseStatementBlock));
-			RulesStatements.Add(FRulePair(EHlslToken::Return, ParseReturnStatement));
-			RulesStatements.Add(FRulePair(EHlslToken::Do, ParseDoStatement));
-			RulesStatements.Add(FRulePair(EHlslToken::While, ParseWhileStatement, TryParseAttribute));
-			RulesStatements.Add(FRulePair(EHlslToken::For, ParseForStatement, TryParseAttribute));
-			RulesStatements.Add(FRulePair(EHlslToken::If, ParseIfStatement, TryParseAttribute));
-			RulesStatements.Add(FRulePair(EHlslToken::Switch, ParseSwitchStatement, TryParseAttribute));
-			RulesStatements.Add(FRulePair(EHlslToken::Semicolon, ParseEmptyStatement));
-			RulesStatements.Add(FRulePair(EHlslToken::Break, ParseEmptyStatement));
-			RulesStatements.Add(FRulePair(EHlslToken::Invalid, ParseLocalDeclaration));
-			// Always try expressions last
-			RulesStatements.Add(FRulePair(EHlslToken::Invalid, ParseExpressionStatement));
+		check(*OutStatement == nullptr);
+
+		auto* Statement = new(Allocator) AST::FJumpStatement(Allocator, AST::EJumpType::Break, Parser.Scanner.PeekToken(-1)->SourceInfo);
+
+		if (Parser.Scanner.MatchToken(EHlslToken::Semicolon))
+		{
+			*OutStatement = Statement;
+			return EParseResult::Matched;
 		}
+
+		return EParseResult::Error;
+	}
+
+	EParseResult ParseContinueStatement(FHlslParser& Parser, FLinearAllocator* Allocator, AST::FNode** OutStatement)
+	{
+		check(*OutStatement == nullptr);
+
+		auto* Statement = new(Allocator) AST::FJumpStatement(Allocator, AST::EJumpType::Continue, Parser.Scanner.PeekToken(-1)->SourceInfo);
+
+		if (Parser.Scanner.MatchToken(EHlslToken::Semicolon))
+		{
+			*OutStatement = Statement;
+			return EParseResult::Matched;
+		}
+
+		return EParseResult::Error;
+	}
+
+	namespace ParserRules
+	{
+		static struct FStaticInitializer
+		{
+			FStaticInitializer()
+			{
+				// Top Level constructs
+				RulesTranslationUnit.Add(FRulePair(EHlslToken::CBuffer, ParseCBuffer));
+				RulesTranslationUnit.Add(FRulePair(EHlslToken::Invalid, ParseFunctionDeclaration, true));
+				RulesTranslationUnit.Add(FRulePair(EHlslToken::Invalid, ParseGlobalVariableDeclaration));
+
+				RulesStatements.Add(FRulePair(EHlslToken::LeftBrace, ParseStatementBlock));
+				RulesStatements.Add(FRulePair(EHlslToken::Return, ParseReturnStatement));
+				RulesStatements.Add(FRulePair(EHlslToken::Do, ParseDoStatement));
+				RulesStatements.Add(FRulePair(EHlslToken::While, ParseWhileStatement, true));
+				RulesStatements.Add(FRulePair(EHlslToken::For, ParseForStatement, true));
+				RulesStatements.Add(FRulePair(EHlslToken::If, ParseIfStatement, true));
+				RulesStatements.Add(FRulePair(EHlslToken::Switch, ParseSwitchStatement, true));
+				RulesStatements.Add(FRulePair(EHlslToken::Semicolon, ParseEmptyStatement));
+				RulesStatements.Add(FRulePair(EHlslToken::Break, ParseBreakStatement));
+				RulesStatements.Add(FRulePair(EHlslToken::Continue, ParseContinueStatement));
+				RulesStatements.Add(FRulePair(EHlslToken::Invalid, ParseLocalDeclaration));
+				// Always try expressions last
+				RulesStatements.Add(FRulePair(EHlslToken::Invalid, ParseExpressionStatement));
+			}
+		} GStaticInitializer;
 	}
 
-	FHlslParser::FHlslParser() :
-		GlobalScope(nullptr)
+	FHlslParser::FHlslParser(FLinearAllocator* InAllocator) :
+		GlobalScope(InAllocator, nullptr),
+		Allocator(InAllocator)
 	{
 		CurrentScope = &GlobalScope;
-		InitRules();
 	}
 
 	namespace Parser
 	{
-		bool Parse(const FString& Input, const FString& Filename)
+		bool Parse(const FString& Input, const FString& Filename, bool bDump)
 		{
-			FHlslParser Parser;
+			FLinearAllocator Allocator;
+			FHlslParser Parser(&Allocator);
 			if (!Parser.Scanner.Lex(Input, Filename))
 			{
 				return false;
 			}
 
 			bool bSuccess = true;
+			TLinearArray<AST::FNode*> Nodes(&Allocator);
 			while (Parser.Scanner.HasMoreTokens())
 			{
 				auto LastIndex = Parser.Scanner.GetCurrentTokenIndex();
 
 				static FString GlobalDeclOrDefinition(TEXT("Global declaration or definition"));
-				auto Result = TryRules(RulesTranslationUnit, Parser, GlobalDeclOrDefinition, true);
+				AST::FNode* Node = nullptr;
+				auto Result = TryRules(RulesTranslationUnit, Parser, GlobalDeclOrDefinition, true, &Allocator, &Node);
 				if (Result == EParseResult::Error)
 				{
 					bSuccess = false;
@@ -1185,6 +1561,11 @@ Unmatched:
 				else
 				{
 					check(Result == EParseResult::Matched);
+					if (bDump && Node)
+					{
+						Node->Dump(0);
+					}
+					Nodes.Add(Node);
 				}
 
 				check(LastIndex != Parser.Scanner.GetCurrentTokenIndex());
diff --git a/Engine/Source/Developer/ShaderCompilerCommon/Private/HlslParser.h b/Engine/Source/Developer/ShaderCompilerCommon/Private/HlslParser.h
--- a/Engine/Source/Developer/ShaderCompilerCommon/Private/HlslParser.h
+++ b/Engine/Source/Developer/ShaderCompilerCommon/Private/HlslParser.h
@@ -20,6 +20,6 @@ namespace CrossCompiler
 	namespace Parser
 	{
 		// Returns true if successfully parsed
-		bool Parse(const FString& Input, const FString& Filename);
+		bool Parse(const FString& Input, const FString& Filename, bool bDump);
 	}
 }
diff --git a/Engine/Source/Programs/CrossCompilerTool/CrossCompilerTool.Target.cs b/Engine/Source/Programs/CrossCompilerTool/CrossCompilerTool.Target.cs
--- a/Engine/Source/Programs/CrossCompilerTool/CrossCompilerTool.Target.cs
+++ b/Engine/Source/Programs/CrossCompilerTool/CrossCompilerTool.Target.cs
@@ -55,5 +55,7 @@ public class CrossCompilerToolTarget : TargetRules
 
 		// UnrealHeaderTool is a console application, not a Windows app (sets entry point to main(), instead of WinMain())
 		OutLinkEnvironmentConfiguration.bIsBuildingConsoleApplication = true;
+
+		//UEBuildConfiguration.bCompileWithStatsWithoutEngine = true;
 	}
 }
diff --git a/Engine/Source/Programs/CrossCompilerTool/Private/CrossCompilerTool.cpp b/Engine/Source/Programs/CrossCompilerTool/Private/CrossCompilerTool.cpp
--- a/Engine/Source/Programs/CrossCompilerTool/Private/CrossCompilerTool.cpp
+++ b/Engine/Source/Programs/CrossCompilerTool/Private/CrossCompilerTool.cpp
@@ -87,7 +87,7 @@ namespace CCT
 					}
 					UE_LOG(LogCrossCompilerTool, Log, TEXT("%d: %s!"), Count++, *File);
 
-					if (!CrossCompiler::Parser::Parse(HLSLShader, File))
+					if (!CrossCompiler::Parser::Parse(HLSLShader, File, false))
 					{
 						UE_LOG(LogCrossCompilerTool, Log, TEXT("Error compiling '%s'!"), *File);
 						return 1;
@@ -98,11 +98,13 @@ namespace CCT
 			{
 				if (RunInfo.bRunCPP)
 				{
+					//GMalloc->DumpAllocatorStats(*FGenericPlatformOutputDevices::GetLog());
 					if (!Preprocess(RunInfo.InputFile, HLSLShaderSource))
 					{
 						UE_LOG(LogCrossCompilerTool, Log, TEXT("Error during preprocessor on '%s'!"), *RunInfo.InputFile);
 						return 1;
 					}
+					//GMalloc->DumpAllocatorStats(*FGenericPlatformOutputDevices::GetLog());
 				}
 				else
 				{
@@ -118,7 +120,7 @@ namespace CCT
 					return 0;
 				}
 
-				if (!CrossCompiler::Parser::Parse(HLSLShaderSource, *RunInfo.InputFile))
+				if (!CrossCompiler::Parser::Parse(HLSLShaderSource, *RunInfo.InputFile, true))
 				{
 					UE_LOG(LogCrossCompilerTool, Log, TEXT("Error compiling '%s'!"), *RunInfo.InputFile);
 					return 1;
@@ -199,5 +201,7 @@ INT32_MAIN_INT32_ARGC_TCHAR_ARGV()
 		return 1;
 	}
 
-	return CCT::Run(RunInfo);
+	int32 Value = CCT::Run(RunInfo);
+	//FGenericPlatformOutputDevices::GetLog()->Flush();
+	return Value;
 }
