UE4 - Fix shader preprocessor replacing #include "Material.usf" to #include "Material_usf"

diff --git a/Engine/Shaders/ShaderVersion.usf b/Engine/Shaders/ShaderVersion.usf
--- a/Engine/Shaders/ShaderVersion.usf
+++ b/Engine/Shaders/ShaderVersion.usf
@@ -2,4 +2,4 @@
 // Each time the console command is executed it generates a new GUID. As this file is included
 // in common.usf (which should be included in any shader) it allows to invalidate the shader DDC.
 // 
-// GUID = {CAC8D6F0-2E83-4C3D-A767-584D1E3DC80y}
+// GUID = {CAC8D6F0-2E83-4C3D-A767-584D1E3DC80t}
diff --git a/Engine/Source/Developer/ShaderCompilerCommon/Private/ShaderCompilerCommon.cpp b/Engine/Source/Developer/ShaderCompilerCommon/Private/ShaderCompilerCommon.cpp
--- a/Engine/Source/Developer/ShaderCompilerCommon/Private/ShaderCompilerCommon.cpp
+++ b/Engine/Source/Developer/ShaderCompilerCommon/Private/ShaderCompilerCommon.cpp
@@ -128,6 +128,83 @@ void BuildResourceTableMapping(
 	OutSRT.MaxBoundResourceTable = MaxBoundResourceTable;
 }
 
+// Specialized version of FString::ReplaceInline that checks that the search text is not inside a #line directive
+static void WholeWordReplaceInline(FString& String, TCHAR* StartPtr, const TCHAR* SearchText, const TCHAR* ReplacementText)
+{
+	if (String.Len() > 0
+		&&	SearchText != nullptr && *SearchText != 0
+		&&	ReplacementText != nullptr && FCString::Strcmp(SearchText, ReplacementText) != 0)
+	{
+		const int32 NumCharsToReplace = FCString::Strlen(SearchText);
+		const int32 NumCharsToInsert = FCString::Strlen(ReplacementText);
+
+		check(NumCharsToInsert == NumCharsToReplace);
+		check(*StartPtr);
+		TCHAR* Pos = FCString::Strstr(StartPtr, SearchText);
+		while (Pos != nullptr)
+		{
+			// Find a " character, indicating we might be inside a #line directive
+			TCHAR* FoundQuote = nullptr;
+			auto* ValidatePos = Pos;
+			do
+			{
+				--ValidatePos;
+				if (*ValidatePos == '\"')
+				{
+					FoundQuote = ValidatePos;
+					break;
+				}
+			}
+			while (ValidatePos >= StartPtr && *ValidatePos != '\n');
+
+			bool bReplace = true;
+			if (FoundQuote)
+			{
+				// Validate that we're indeed inside a #line directive by first finding the last \n character
+				TCHAR* FoundEOL = nullptr;
+				do 
+				{
+					--ValidatePos;
+					if (*ValidatePos == '\n')
+					{
+						FoundEOL = ValidatePos;
+						break;
+					}
+				}
+				while (ValidatePos > StartPtr);
+
+				// Finally make sure the directive is between the \n and the and the quote
+				if (FoundEOL)
+				{
+					auto* FoundInclude = FCString::Strstr(FoundEOL + 1, TEXT("#line"));
+					if (FoundInclude && FoundInclude < FoundQuote)
+					{
+						bReplace = false;
+					}
+				}
+			}
+			
+			if (bReplace)
+			{
+				// FCString::Strcpy inserts a terminating zero so can't use that
+				for (int32 i = 0; i < NumCharsToInsert; i++)
+				{
+					Pos[i] = ReplacementText[i];
+				}
+			}
+
+			if (Pos + NumCharsToReplace - *String < String.Len())
+			{
+				Pos = FCString::Strstr(Pos + NumCharsToReplace, SearchText);
+			}
+			else
+			{
+				break;
+			}
+		}
+	}
+}
+
 
 bool RemoveUniformBuffersFromSource(FString& SourceCode)
 {
@@ -173,8 +250,7 @@ bool RemoveUniformBuffersFromSource(FString& SourceCode)
 		// Find & Replace this UB
 		FString UBSource = UniformBufferName + FString(TEXT("."));
 		FString UBDest = UniformBufferName + FString(TEXT("_"));
-		SourceCode.ReplaceInline(*UBSource, *UBDest, ESearchCase::CaseSensitive);
-
+		WholeWordReplaceInline(SourceCode, &SourceCharArray[CloseBraceSemicolonTokenPos + 2], *UBSource, *UBDest);
 
 		// Find next UB
 		StaticStructTokenPos = SourceCode.Find(StaticStructToken, ESearchCase::CaseSensitive, ESearchDir::FromStart, CloseBraceSemicolonTokenPos + 2);
