Added location for Linux CCT and got rid of preprocessor usage.

diff --git a/Engine/Source/Developer/ShaderCompilerCommon/Private/ShaderCompilerCommon.cpp b/Engine/Source/Developer/ShaderCompilerCommon/Private/ShaderCompilerCommon.cpp
--- a/Engine/Source/Developer/ShaderCompilerCommon/Private/ShaderCompilerCommon.cpp
+++ b/Engine/Source/Developer/ShaderCompilerCommon/Private/ShaderCompilerCommon.cpp
@@ -262,15 +262,28 @@ bool RemoveUniformBuffersFromSource(FString& SourceCode)
 
 FString CreateCrossCompilerBatchFileContents(const FString& ShaderFile, const FString& OutputFile, const FString& FrequencySwitch, const FString& EntryPoint, const FString& VersionSwitch, const FString& ExtraArguments)
 {
-#if PLATFORM_MAC || PLATFORM_LINUX
-	FString BatchFile = FPaths::RootDir() / FString::Printf(TEXT("Engine/Source/ThirdParty/hlslcc/hlslcc/bin/Mac/hlslcc_64 %s -o=%s %s -entry=%s %s %s"), *ShaderFile, *OutputFile, *FrequencySwitch, *EntryPoint, *VersionSwitch, *ExtraArguments);
-#elif PLATFORM_WINDOWS
-	FString BatchFile = TEXT("@echo off");
-	BatchFile += TEXT("\nif defined ue.hlslcc GOTO DONE\nset ue.hlslcc=");
-	BatchFile += FPaths::RootDir() / TEXT("Engine\\Binaries\\Win64\\CrossCompilerTool.exe");
-	BatchFile += TEXT("\n\n:DONE\n%ue.hlslcc% ");
-	BatchFile += FString::Printf(TEXT("\"%s\" -o=\"%s\" %s -entry=%s %s %s"), *ShaderFile, *OutputFile, *FrequencySwitch, *EntryPoint, *VersionSwitch, *ExtraArguments);
-	BatchFile += TEXT("\npause\n");
-#endif
+	FString BatchFile;
+	if (PLATFORM_MAC)
+	{
+		BatchFile = FPaths::RootDir() / FString::Printf(TEXT("Engine/Source/ThirdParty/hlslcc/hlslcc/bin/Mac/hlslcc_64 %s -o=%s %s -entry=%s %s %s"), *ShaderFile, *OutputFile, *FrequencySwitch, *EntryPoint, *VersionSwitch, *ExtraArguments);
+	}
+	else if (PLATFORM_LINUX)
+	{
+		BatchFile = FPaths::RootDir() / FString::Printf(TEXT("Engine/Binaries/Linux/CrossCompilerTool %s -o=%s %s -entry=%s %s %s"), *ShaderFile, *OutputFile, *FrequencySwitch, *EntryPoint, *VersionSwitch, *ExtraArguments);
+	}
+	else if (PLATFORM_WINDOWS)
+	{
+		FString BatchFile = TEXT("@echo off");
+		BatchFile += TEXT("\nif defined ue.hlslcc GOTO DONE\nset ue.hlslcc=");
+		BatchFile += FPaths::RootDir() / TEXT("Engine\\Binaries\\Win64\\CrossCompilerTool.exe");
+		BatchFile += TEXT("\n\n:DONE\n%ue.hlslcc% ");
+		BatchFile += FString::Printf(TEXT("\"%s\" -o=\"%s\" %s -entry=%s %s %s"), *ShaderFile, *OutputFile, *FrequencySwitch, *EntryPoint, *VersionSwitch, *ExtraArguments);
+		BatchFile += TEXT("\npause\n");
+	}
+	else
+	{
+		checkf(false, TEXT("CreateCrossCompilerBatchFileContents: unsupported platform!"));
+	}
+
 	return BatchFile;
 }
