CCT - Support for uniform token

diff --git a/Engine/Source/Developer/ShaderCompilerCommon/Private/HlslAST.h b/Engine/Source/Developer/ShaderCompilerCommon/Private/HlslAST.h
--- a/Engine/Source/Developer/ShaderCompilerCommon/Private/HlslAST.h
+++ b/Engine/Source/Developer/ShaderCompilerCommon/Private/HlslAST.h
@@ -495,6 +495,7 @@ namespace CrossCompiler
 					uint32 bOut : 1;
 					uint32 bRowMajor : 1;
 					uint32 bShared : 1;
+					uint32 bUniform : 1;
 				};
 				uint32 Raw;
 			};
diff --git a/Engine/Source/Developer/ShaderCompilerCommon/Private/HlslLexer.cpp b/Engine/Source/Developer/ShaderCompilerCommon/Private/HlslLexer.cpp
--- a/Engine/Source/Developer/ShaderCompilerCommon/Private/HlslLexer.cpp
+++ b/Engine/Source/Developer/ShaderCompilerCommon/Private/HlslLexer.cpp
@@ -377,6 +377,7 @@ namespace CrossCompiler
 				InsertToken(TEXT("out"), EHlslToken::Out);
 				InsertToken(TEXT("inout"), EHlslToken::InOut);
 				InsertToken(TEXT("static"), EHlslToken::Static);
+				InsertToken(TEXT("uniform"), EHlslToken::Uniform);
 
 				// Misc
 				InsertToken(TEXT("["), EHlslToken::LeftSquareBracket);
diff --git a/Engine/Source/Developer/ShaderCompilerCommon/Private/HlslLexer.h b/Engine/Source/Developer/ShaderCompilerCommon/Private/HlslLexer.h
--- a/Engine/Source/Developer/ShaderCompilerCommon/Private/HlslLexer.h
+++ b/Engine/Source/Developer/ShaderCompilerCommon/Private/HlslLexer.h
@@ -227,6 +227,7 @@ namespace CrossCompiler
 		Out,
 		InOut,
 		Static,
+		Uniform,
 
 		// Misc
 		LeftSquareBracket,
diff --git a/Engine/Source/Developer/ShaderCompilerCommon/Private/HlslParser.cpp b/Engine/Source/Developer/ShaderCompilerCommon/Private/HlslParser.cpp
--- a/Engine/Source/Developer/ShaderCompilerCommon/Private/HlslParser.cpp
+++ b/Engine/Source/Developer/ShaderCompilerCommon/Private/HlslParser.cpp
@@ -245,19 +245,19 @@ namespace CrossCompiler
 	// Multi declaration parser flags
 	enum EDeclarationFlags
 	{
-		//EDF_ROLLBACK_IF_NO_MATCH		= 0x0001,
-		EDF_CONST_ROW_MAJOR				= 0x0002,
-		EDF_STATIC						= 0x0004,
-		EDF_TEXTURE_SAMPLER_OR_BUFFER	= 0x0008,
-		EDF_INITIALIZER					= 0x0010,
-		EDF_INITIALIZER_LIST			= 0x0020 | EDF_INITIALIZER,
-		EDF_SEMANTIC					= 0x0040,
-		EDF_SEMICOLON					= 0x0080,
-		EDF_IN_OUT						= 0x0100,
-		EDF_MULTIPLE					= 0x0200,
-		EDF_PRIMITIVE_DATA_TYPE			= 0x0400,
-		EDF_SHARED						= 0x0800,
-		EDF_NOINTERPOLATION				= 0x1000,
+		EDF_CONST_ROW_MAJOR				= 1 << 0,
+		EDF_STATIC						= 1 << 1,
+		EDF_UNIFORM						= 1 << 2,
+		EDF_TEXTURE_SAMPLER_OR_BUFFER	= 1 << 3,
+		EDF_INITIALIZER					= 1 << 4,
+		EDF_INITIALIZER_LIST				= (1 << 5) | EDF_INITIALIZER,
+		EDF_SEMANTIC						= 1 << 6,
+		EDF_SEMICOLON					= 1 << 7,
+		EDF_IN_OUT						= 1 << 8,
+		EDF_MULTIPLE						= 1 << 9,
+		EDF_PRIMITIVE_DATA_TYPE			= 1 << 10,
+		EDF_SHARED						= 1 << 11,
+		EDF_NOINTERPOLATION				= 1 << 12,
 	};
 
 	EParseResult ParseInitializer(FHlslScanner& Scanner, FSymbolScope* SymbolScope, bool bAllowLists, FLinearAllocator* Allocator, AST::FExpression** OutList)
@@ -300,6 +300,7 @@ namespace CrossCompiler
 		int32 OutFound = 0;
 		int32 InOutFound = 0;
 		int32 PrimitiveFound = 0;
+		int32 UniformFound = 0;
 
 		if (Flags & EDF_PRIMITIVE_DATA_TYPE)
 		{
@@ -416,15 +417,31 @@ namespace CrossCompiler
 					return EParseResult::Error;
 				}
 			}
+			else if ((Flags & EDF_UNIFORM) && Scanner.MatchToken(EHlslToken::Uniform))
+			{
+				++UniformFound;
+				Qualifier->bUniform = true;
+				if (UniformFound > 1)
+				{
+					Scanner.SourceError(TEXT("'uniform' found more than once!\n"));
+					return EParseResult::Error;
+				}
+			}
 			else
 			{
 				break;
 			}
 		}
 
+		if (UniformFound && (OutFound || InOutFound || PrimitiveFound || SharedFound || NoInterpolationFound))
+		{
+			Scanner.SourceError(TEXT("'uniform' can not be used with other storage qualifiers (inout, out, nointerpolation, etc)!\n"));
+			return EParseResult::Error;
+		}
+
 		bOutPrimitiveFound = (PrimitiveFound > 0);
 
-		return (ConstFound + RowMajorFound + InFound + OutFound + InOutFound + StaticFound + SharedFound + PrimitiveFound + NoInterpolationFound)
+		return (ConstFound + RowMajorFound + InFound + OutFound + InOutFound + StaticFound + SharedFound + PrimitiveFound + NoInterpolationFound + UniformFound)
 			? EParseResult::Matched
 			: EParseResult::NotMatched;
 	}
@@ -765,7 +782,7 @@ namespace CrossCompiler
 		while (Parser.Scanner.HasMoreTokens())
 		{
 			AST::FDeclaratorList* Declaration = nullptr;
-			auto Result = ParseGeneralDeclaration(Parser.Scanner, Parser.CurrentScope, Allocator, &Declaration, EDF_CONST_ROW_MAJOR | EDF_IN_OUT | EDF_TEXTURE_SAMPLER_OR_BUFFER | EDF_INITIALIZER | EDF_SEMANTIC | EDF_PRIMITIVE_DATA_TYPE | EDF_NOINTERPOLATION);
+			auto Result = ParseGeneralDeclaration(Parser.Scanner, Parser.CurrentScope, Allocator, &Declaration, EDF_CONST_ROW_MAJOR | EDF_IN_OUT | EDF_TEXTURE_SAMPLER_OR_BUFFER | EDF_INITIALIZER | EDF_SEMANTIC | EDF_PRIMITIVE_DATA_TYPE | EDF_NOINTERPOLATION | EDF_UNIFORM);
 			if (Result == EParseResult::Error)
 			{
 				return EParseResult::Error;
@@ -957,7 +974,7 @@ check(0);
 	EParseResult ParseGlobalVariableDeclaration(FHlslParser& Parser, FLinearAllocator* Allocator, AST::FNode** OutDeclaration)
 	{
 		AST::FDeclaratorList* List = nullptr;
-		auto Result = ParseGeneralDeclaration(Parser.Scanner, Parser.CurrentScope, Allocator, &List, EDF_CONST_ROW_MAJOR | EDF_STATIC | EDF_SHARED | EDF_TEXTURE_SAMPLER_OR_BUFFER | EDF_INITIALIZER | EDF_INITIALIZER_LIST | EDF_SEMICOLON | EDF_MULTIPLE);
+		auto Result = ParseGeneralDeclaration(Parser.Scanner, Parser.CurrentScope, Allocator, &List, EDF_CONST_ROW_MAJOR | EDF_STATIC | EDF_SHARED | EDF_TEXTURE_SAMPLER_OR_BUFFER | EDF_INITIALIZER | EDF_INITIALIZER_LIST | EDF_SEMICOLON | EDF_MULTIPLE | EDF_UNIFORM);
 		*OutDeclaration = List;
 		return Result;
 	}
diff --git a/Engine/Source/Programs/CrossCompilerTool/Private/CrossCompilerTool.cpp b/Engine/Source/Programs/CrossCompilerTool/Private/CrossCompilerTool.cpp
--- a/Engine/Source/Programs/CrossCompilerTool/Private/CrossCompilerTool.cpp
+++ b/Engine/Source/Programs/CrossCompilerTool/Private/CrossCompilerTool.cpp
@@ -85,7 +85,7 @@ namespace CCT
 						UE_LOG(LogCrossCompilerTool, Error, TEXT("Couldn't load Input file '%s'!"), *File);
 						continue;
 					}
-					UE_LOG(LogCrossCompilerTool, Log, TEXT("%d: %s!"), Count++, *File);
+					UE_LOG(LogCrossCompilerTool, Log, TEXT("%d: %s"), Count++, *File);
 
 					if (!CrossCompiler::Parser::Parse(HLSLShader, File, false))
 					{
