CCT - Move allocator to HlslUtils
- Add interpolation modifiers
- Ignore #pragma
- Fix for strings in attributes
- Do translation units separately

diff --git a/Engine/Source/Developer/ShaderCompilerCommon/Private/HlslAST.cpp b/Engine/Source/Developer/ShaderCompilerCommon/Private/HlslAST.cpp
--- a/Engine/Source/Developer/ShaderCompilerCommon/Private/HlslAST.cpp
+++ b/Engine/Source/Developer/ShaderCompilerCommon/Private/HlslAST.cpp
@@ -542,7 +542,7 @@ namespace CrossCompiler
 		{
 			DumpIndent(Indent);
 			DumpAttributes();
-			FPlatformMisc::LowLevelOutputDebugString(TEXT("if("));
+			FPlatformMisc::LowLevelOutputDebugString(TEXT("if ("));
 			Condition->Dump(Indent);
 			FPlatformMisc::LowLevelOutputDebugString(TEXT(")\n"));
 			ThenStatement->Dump(Indent);
@@ -566,6 +566,8 @@ namespace CrossCompiler
 
 		FTypeSpecifier::FTypeSpecifier(FLinearAllocator* InAllocator, const FSourceInfo& InInfo) :
 			FNode(InAllocator, InInfo),
+			TypeName(nullptr),
+			InnerType(nullptr),
 			Structure(nullptr),
 			TextureMSNumSamples(1),
 			PatchSize(0),
@@ -641,7 +643,7 @@ namespace CrossCompiler
 			}
 
 			DumpIndent(Indent);
-			FPlatformMisc::LowLevelOutputDebugString(TEXT("}\n"));
+			FPlatformMisc::LowLevelOutputDebugString(TEXT("}\n\n"));
 		}
 
 		FCBufferDeclaration::~FCBufferDeclaration()
@@ -669,7 +671,11 @@ namespace CrossCompiler
 				FPlatformMisc::LowLevelOutputDebugString(TEXT("static "));
 			}
 
-			if (bIn && bOut)
+			if (bShared)
+			{
+				FPlatformMisc::LowLevelOutputDebugString(TEXT("groupshared "));
+			}
+			else if (bIn && bOut)
 			{
 				FPlatformMisc::LowLevelOutputDebugString(TEXT("inout "));
 			}
@@ -682,8 +688,31 @@ namespace CrossCompiler
 				FPlatformMisc::LowLevelOutputDebugString(TEXT("out "));
 			}
 
-			//bRowMajor?
-			//bShared?
+			if (bLinear)
+			{
+				FPlatformMisc::LowLevelOutputDebugString(TEXT("linear "));
+			}
+			if (bCentroid)
+			{
+				FPlatformMisc::LowLevelOutputDebugString(TEXT("centroid "));
+			}
+			if (bNoInterpolation)
+			{
+				FPlatformMisc::LowLevelOutputDebugString(TEXT("nointerpolation "));
+			}
+			if (bNoPerspective)
+			{
+				FPlatformMisc::LowLevelOutputDebugString(TEXT("noperspective "));
+			}
+			if (bSample)
+			{
+				FPlatformMisc::LowLevelOutputDebugString(TEXT("sample "));
+			}
+
+			if (bRowMajor)
+			{
+				FPlatformMisc::LowLevelOutputDebugString(TEXT("row_major "));
+			}
 		}
 
 		FFullySpecifiedType::FFullySpecifiedType(FLinearAllocator* InAllocator, const FSourceInfo& InInfo) :
@@ -758,7 +787,7 @@ namespace CrossCompiler
 			DumpAttributes();
 			if (Type)
 			{
-				Type->Dump(0);
+				Type->Dump(Indent);
 			}
 
 			bool bFirst = true;
@@ -912,14 +941,17 @@ namespace CrossCompiler
 					RestExpression->Dump(0);
 				}
 				FPlatformMisc::LowLevelOutputDebugString(TEXT(")\n"));
-				DumpIndent(Indent);
-				FPlatformMisc::LowLevelOutputDebugString(TEXT("{\n"));
 				if (Body)
 				{
 					Body->Dump(Indent + 1);
 				}
-				DumpIndent(Indent);
-				FPlatformMisc::LowLevelOutputDebugString(TEXT("}\n"));
+				else
+				{
+					DumpIndent(Indent);
+					FPlatformMisc::LowLevelOutputDebugString(TEXT("{\n"));
+					DumpIndent(Indent);
+					FPlatformMisc::LowLevelOutputDebugString(TEXT("}\n"));
+				}
 				break;
 
 			case EIterationType::While:
@@ -1172,8 +1204,7 @@ namespace CrossCompiler
 
 		void FStructSpecifier::Dump(int32 Indent) const
 		{
-			DumpIndent(Indent);
-			FPlatformMisc::LowLevelOutputDebugStringf(TEXT("struct %s"), Name);
+			FPlatformMisc::LowLevelOutputDebugStringf(TEXT("struct %s"), Name ? Name : TEXT(""));
 			if (ParentName && *ParentName)
 			{
 				FPlatformMisc::LowLevelOutputDebugStringf(TEXT(" : %s"), ParentName);
diff --git a/Engine/Source/Developer/ShaderCompilerCommon/Private/HlslAST.h b/Engine/Source/Developer/ShaderCompilerCommon/Private/HlslAST.h
--- a/Engine/Source/Developer/ShaderCompilerCommon/Private/HlslAST.h
+++ b/Engine/Source/Developer/ShaderCompilerCommon/Private/HlslAST.h
@@ -8,190 +8,8 @@
 
 #include "HlslLexer.h"
 
-#define USE_UNREAL_ALLOCATOR 0
-
 namespace CrossCompiler
 {
-	struct FLinearAllocator
-	{
-		static const int PageSize = 1024* 1024;
-
-		FLinearAllocator()
-		{
-			auto* Initial = new FPage(PageSize);
-			Pages.Add(Initial);
-		}
-
-		~FLinearAllocator()
-		{
-			for (auto* Page : Pages)
-			{
-				delete Page;
-			}
-		}
-
-		inline void* Alloc(SIZE_T NumBytes)
-		{
-			check(NumBytes <= PageSize);
-			auto* Page = Pages.Last();
-			if (Page->Current + NumBytes > Page->End)
-			{
-				Page = new FPage(PageSize);
-				Pages.Add(Page);
-			}
-
-			void* Ptr = Page->Current;
-			Page->Current += NumBytes;
-			return Ptr;
-		}
-
-		inline void* Alloc(SIZE_T NumBytes, SIZE_T Align)
-		{
-			void* Data = Alloc(NumBytes + Align - 1);
-			UPTRINT Address = (UPTRINT)Data;
-			Address += (Align - (Address % (UPTRINT)Align)) % Align;
-			return (void*)Address;
-		}
-
-		inline TCHAR* Strdup(const TCHAR* String)
-		{
-			if (!String)
-			{
-				return nullptr;
-			}
-
-			auto* Data = (TCHAR*)Alloc(FCString::Strlen(String) + 1);
-			FCString::Strcpy(Data, FCString::Strlen(String) + 1, String);
-			return Data;
-		}
-
-		inline TCHAR* Strdup(const FString& String)
-		{
-			auto* Data = (TCHAR*)Alloc(String.Len() + 1);
-			FCString::Strcpy(Data, String.Len() + 1, *String);
-			return Data;
-		}
-
-		struct FPage
-		{
-			char* Current;
-			char* Begin;
-			char* End;
-
-			FPage(SIZE_T Size)
-			{
-				Begin = new char[Size];
-				End = Begin + Size;
-				Current = Begin;
-			}
-
-			~FPage()
-			{
-				delete [] Begin;
-			}
-		};
-		TArray<FPage*> Pages;
-	};
-
-#if !USE_UNREAL_ALLOCATOR
-	class FLinearAllocatorPolicy
-	{
-	public:
-		// Unreal allocator magic
-		enum { NeedsElementType = false };
-		enum { RequireRangeCheck = true };
-
-		template<typename ElementType>
-		class ForElementType
-		{
-		public:
-
-			/** Default constructor. */
-			ForElementType() :
-				LinearAllocator(nullptr),
-				Data(nullptr)
-			{}
-
-			// FContainerAllocatorInterface
-			/*FORCEINLINE*/ ElementType* GetAllocation() const
-			{
-				return Data;
-			}
-			void ResizeAllocation(int32 PreviousNumElements, int32 NumElements, int32 NumBytesPerElement)
-			{
-				void* OldData = Data;
-				if (NumElements)
-				{
-					// Allocate memory from the stack.
-					Data = (ElementType*)LinearAllocator->Alloc(NumElements * NumBytesPerElement,
-						FMath::Max((uint32)sizeof(void*), (uint32)ALIGNOF(ElementType))
-						);
-
-					// If the container previously held elements, copy them into the new allocation.
-					if (OldData && PreviousNumElements)
-					{
-						const int32 NumCopiedElements = FMath::Min(NumElements, PreviousNumElements);
-						FMemory::Memcpy(Data, OldData, NumCopiedElements * NumBytesPerElement);
-					}
-				}
-			}
-			int32 CalculateSlack(int32 NumElements, int32 NumAllocatedElements, int32 NumBytesPerElement) const
-			{
-				return DefaultCalculateSlack(NumElements, NumAllocatedElements, NumBytesPerElement);
-			}
-
-			int32 GetAllocatedSize(int32 NumAllocatedElements, int32 NumBytesPerElement) const
-			{
-				return NumAllocatedElements * NumBytesPerElement;
-			}
-
-			FLinearAllocator* LinearAllocator;
-
-		private:
-
-			/** A pointer to the container's elements. */
-			ElementType* Data;
-		};
-
-		typedef ForElementType<FScriptContainerElement> ForAnyElementType;
-	};
-
-	class FLinearBitArrayAllocator
-		: public TInlineAllocator<4, FLinearAllocatorPolicy>
-	{
-	};
-
-	class FLinearSparseArrayAllocator
-		: public TSparseArrayAllocator<FLinearAllocatorPolicy, FLinearBitArrayAllocator>
-	{
-	};
-
-	class FLinearSetAllocator
-		: public TSetAllocator<FLinearSparseArrayAllocator, TInlineAllocator<1, FLinearAllocatorPolicy> >
-	{
-	};
-
-	template <typename TType>
-	class TLinearArray : public TArray<TType, FLinearAllocatorPolicy>
-	{
-	public:
-		TLinearArray(FLinearAllocator* Allocator)
-		{
-			TArray<TType, FLinearAllocatorPolicy>::AllocatorInstance.LinearAllocator = Allocator;
-		}
-	};
-
-/*
-	template <typename TType>
-	struct TLinearSet : public TSet<typename TType, DefaultKeyFuncs<typename TType>, FLinearSetAllocator>
-	{
-		TLinearSet(FLinearAllocator* InAllocator)
-		{
-			Elements.AllocatorInstance.LinearAllocator = InAllocator;
-		}
-	};*/
-#endif
-
 	namespace AST
 	{
 		class FNode
@@ -496,6 +314,13 @@ namespace CrossCompiler
 					uint32 bRowMajor : 1;
 					uint32 bShared : 1;
 					uint32 bUniform : 1;
+
+					// Interpolation modifiers
+					uint32 bLinear : 1;
+					uint32 bCentroid : 1;
+					uint32 bNoInterpolation : 1;
+					uint32 bNoPerspective : 1;
+					uint32 bSample : 1;
 				};
 				uint32 Raw;
 			};
diff --git a/Engine/Source/Developer/ShaderCompilerCommon/Private/HlslExpressionParser.inl b/Engine/Source/Developer/ShaderCompilerCommon/Private/HlslExpressionParser.inl
--- a/Engine/Source/Developer/ShaderCompilerCommon/Private/HlslExpressionParser.inl
+++ b/Engine/Source/Developer/ShaderCompilerCommon/Private/HlslExpressionParser.inl
@@ -118,10 +118,11 @@ namespace CrossCompiler
 
 	enum ETypeFlags
 	{
-		ETF_VOID					= 0x0001,
-		ETF_BUILTIN_NUMERIC			= 0x0002,
-		ETF_SAMPLER_TEXTURE_BUFFER	= 0x0004,
-		ETF_USER_TYPES				= 0x0008,
+		ETF_VOID					= 1 << 0,
+		ETF_BUILTIN_NUMERIC			= 1 << 1,
+		ETF_SAMPLER_TEXTURE_BUFFER	= 1 << 2,
+		ETF_USER_TYPES				= 1 << 3,
+		ETF_ERROR_IF_NOT_USER_TYPE	= 1 << 4,
 	};
 
 	EParseResult ParseGeneralType(const FHlslToken* Token, int32 TypeFlags, FLinearAllocator* Allocator, AST::FTypeSpecifier** OutSpecifier)
@@ -309,7 +310,7 @@ namespace CrossCompiler
 		return EParseResult::NotMatched;
 	}
 
-	EParseResult ParseGeneralType(const FHlslToken* Token, int32 TypeFlags, FSymbolScope* SymbolScope, FLinearAllocator* Allocator, AST::FTypeSpecifier** OutSpecifier)
+	EParseResult ParseGeneralTypeFromToken(const FHlslToken* Token, int32 TypeFlags, FSymbolScope* SymbolScope, FLinearAllocator* Allocator, AST::FTypeSpecifier** OutSpecifier)
 	{
 		if (Token)
 		{
@@ -321,12 +322,19 @@ namespace CrossCompiler
 			if (TypeFlags & ETF_USER_TYPES)
 			{
 				check(SymbolScope);
-				if (Token->Token == EHlslToken::Identifier && FSymbolScope::FindType(SymbolScope, Token->String))
+				if (Token->Token == EHlslToken::Identifier)
 				{
-					auto* Type = new(Allocator) AST::FTypeSpecifier(Allocator, Token->SourceInfo);
-					Type->TypeName = Allocator->Strdup(Token->String);
-					*OutSpecifier = Type;
-					return EParseResult::Matched;
+					if (FSymbolScope::FindType(SymbolScope, Token->String))
+					{
+						auto* Type = new(Allocator)AST::FTypeSpecifier(Allocator, Token->SourceInfo);
+						Type->TypeName = Allocator->Strdup(Token->String);
+						*OutSpecifier = Type;
+						return EParseResult::Matched;
+					}
+					else if (TypeFlags & ETF_ERROR_IF_NOT_USER_TYPE)
+					{
+						return EParseResult::Error;
+					}
 				}
 			}
 
@@ -339,12 +347,19 @@ namespace CrossCompiler
 	EParseResult ParseGeneralType(FHlslScanner& Scanner, int32 TypeFlags, FSymbolScope* SymbolScope, FLinearAllocator* Allocator, AST::FTypeSpecifier** OutSpecifier)
 	{
 		auto* Token = Scanner.PeekToken();
-		if (ParseGeneralType(Token, TypeFlags, SymbolScope, Allocator, OutSpecifier) == EParseResult::Matched)
+		auto Result = ParseGeneralTypeFromToken(Token, TypeFlags, SymbolScope, Allocator, OutSpecifier);
+		if (Result == EParseResult::Matched)
 		{
 			Scanner.Advance();
 			return EParseResult::Matched;
 		}
 
+		if (Result == EParseResult::Error && (TypeFlags & ETF_ERROR_IF_NOT_USER_TYPE))
+		{
+			Scanner.SourceError(*FString::Printf(TEXT("Unknown type '%s'!"), *Token->String));
+			return Result;
+		}
+
 		return EParseResult::NotMatched;
 	}
 
@@ -403,7 +418,7 @@ namespace CrossCompiler
 				const auto* Peek1 = Scanner.PeekToken(1);
 				const auto* Peek2 = Scanner.PeekToken(2);
 				AST::FTypeSpecifier* TypeSpecifier = nullptr;
-				if (Peek1 && ParseGeneralType(Peek1, ETF_BUILTIN_NUMERIC | ETF_USER_TYPES, SymbolScope, Allocator, &TypeSpecifier) == EParseResult::Matched && Peek2 && Peek2->Token == EHlslToken::RightParenthesis)
+				if (Peek1 && ParseGeneralTypeFromToken(Peek1, ETF_BUILTIN_NUMERIC | ETF_USER_TYPES, SymbolScope, Allocator, &TypeSpecifier) == EParseResult::Matched && Peek2 && Peek2->Token == EHlslToken::RightParenthesis)
 				{
 					// Cast
 					Scanner.Advance();
diff --git a/Engine/Source/Developer/ShaderCompilerCommon/Private/HlslLexer.cpp b/Engine/Source/Developer/ShaderCompilerCommon/Private/HlslLexer.cpp
--- a/Engine/Source/Developer/ShaderCompilerCommon/Private/HlslLexer.cpp
+++ b/Engine/Source/Developer/ShaderCompilerCommon/Private/HlslLexer.cpp
@@ -369,8 +369,6 @@ namespace CrossCompiler
 				InsertToken(TEXT("RWTexture2DArray"), EHlslToken::RWTexture2DArray);
 				InsertToken(TEXT("RWTexture3D"), EHlslToken::RWTexture3D);
 				InsertToken(TEXT("StructuredBuffer"), EHlslToken::StructuredBuffer);
-				InsertToken(TEXT("InputPatch"), EHlslToken::InputPatch);
-				InsertToken(TEXT("OutputPatch"), EHlslToken::OutputPatch);
 
 				// Modifiers
 				InsertToken(TEXT("in"), EHlslToken::In);
@@ -389,7 +387,6 @@ namespace CrossCompiler
 				InsertToken(TEXT("struct"), EHlslToken::Struct);
 				InsertToken(TEXT("cbuffer"), EHlslToken::CBuffer);
 				InsertToken(TEXT("groupshared"), EHlslToken::GroupShared);
-				InsertToken(TEXT("nointerpolation"), EHlslToken::NoInterpolation);
 				InsertToken(TEXT("row_major"), EHlslToken::RowMajor);
 			}
 		} GStaticInitializer;
@@ -752,9 +749,14 @@ namespace CrossCompiler
 					check(0);
 				}
 			}
+			else if (Tokenizer.MatchString(MATCH_TARGET(TEXT("#pragma"))))
+			{
+				//@todo-rco: Pragma!
+				SourceWarning(*FString::Printf(TEXT("Ignoring pragma!")));
+			}
 			else
 			{
-				//@todo-rco: Warn about unknown pragma
+				//@todo-rco: Warn about unknown preprocessor directive
 				check(0);
 			}
 
@@ -906,7 +908,7 @@ namespace CrossCompiler
 		Clear(Filename);
 		
 		// Simple heuristic to avoid reallocating
-		Tokens.Reserve(String.Len() / 4);
+		Tokens.Reserve(String.Len() / 8);
 
 		FTokenizer Tokenizer(String);
 		while (Tokenizer.HasCharsAvailable())
@@ -958,6 +960,10 @@ namespace CrossCompiler
 				{
 					AddToken(FHlslToken(SymbolToken, Identifier), Tokenizer);
 				}
+				else if (Tokenizer.MatchQuotedString(Identifier))
+				{
+					AddToken(FHlslToken(EHlslToken::StringConstant, Identifier), Tokenizer);
+				}
 				else if (Tokenizer.HasCharsAvailable())
 				{
 					//@todo-rco: Unknown token!
diff --git a/Engine/Source/Developer/ShaderCompilerCommon/Private/HlslLexer.h b/Engine/Source/Developer/ShaderCompilerCommon/Private/HlslLexer.h
--- a/Engine/Source/Developer/ShaderCompilerCommon/Private/HlslLexer.h
+++ b/Engine/Source/Developer/ShaderCompilerCommon/Private/HlslLexer.h
@@ -6,6 +6,8 @@
 
 #pragma once
 
+#include "HlslUtils.h"
+
 namespace CrossCompiler
 {
 	enum class EHlslToken
@@ -219,8 +221,6 @@ namespace CrossCompiler
 		RWTexture2DArray,
 		RWTexture3D,
 		StructuredBuffer,
-		InputPatch,
-		OutputPatch,
 
 		// Modifiers
 		In,
@@ -239,7 +239,6 @@ namespace CrossCompiler
 		Struct,
 		CBuffer,
 		GroupShared,
-		NoInterpolation,
 		RowMajor,
 
 		Identifier,
@@ -249,15 +248,6 @@ namespace CrossCompiler
 		StringConstant,	// C-style "string"
 	};
 
-	struct FSourceInfo
-	{
-		FString* Filename;
-		int32 Line;
-		int32 Column;
-
-		FSourceInfo() : Filename(nullptr), Line(0), Column(0) {}
-	};
-
 	struct FHlslToken
 	{
 		EHlslToken Token;
diff --git a/Engine/Source/Developer/ShaderCompilerCommon/Private/HlslParser.cpp b/Engine/Source/Developer/ShaderCompilerCommon/Private/HlslParser.cpp
--- a/Engine/Source/Developer/ShaderCompilerCommon/Private/HlslParser.cpp
+++ b/Engine/Source/Developer/ShaderCompilerCommon/Private/HlslParser.cpp
@@ -35,12 +35,11 @@ namespace CrossCompiler
 		FLinearAllocator* Allocator;
 	};
 
-	TRulesArray RulesTranslationUnit;
 	TRulesArray RulesStatements;
 
-	EParseResult TryRules(const TRulesArray& Rules, FHlslParser& Parser, const FString& RuleNames, bool bErrorIfNoMatch, FLinearAllocator* Allocator, AST::FNode** OutNode)
+	EParseResult TryStatementRules(FHlslParser& Parser, FLinearAllocator* Allocator, AST::FNode** OutNode)
 	{
-		for (const auto& Rule : Rules)
+		for (const auto& Rule : RulesStatements)
 		{
 			auto CurrentTokenIndex = Parser.Scanner.GetCurrentTokenIndex();
 			TLinearArray<AST::FAttribute*> Attributes(Allocator);
@@ -90,12 +89,6 @@ namespace CrossCompiler
 			Parser.Scanner.SetCurrentTokenIndex(CurrentTokenIndex);
 		}
 
-		if (bErrorIfNoMatch)
-		{
-			Parser.Scanner.SourceError(FString::Printf(TEXT("No matching %s rules found!"), *RuleNames));
-			return EParseResult::Error;
-		}
-
 		return EParseResult::NotMatched;
 	}
 
@@ -245,19 +238,19 @@ namespace CrossCompiler
 	// Multi declaration parser flags
 	enum EDeclarationFlags
 	{
-		EDF_CONST_ROW_MAJOR				= 1 << 0,
-		EDF_STATIC						= 1 << 1,
-		EDF_UNIFORM						= 1 << 2,
-		EDF_TEXTURE_SAMPLER_OR_BUFFER	= 1 << 3,
-		EDF_INITIALIZER					= 1 << 4,
-		EDF_INITIALIZER_LIST				= (1 << 5) | EDF_INITIALIZER,
-		EDF_SEMANTIC						= 1 << 6,
-		EDF_SEMICOLON					= 1 << 7,
-		EDF_IN_OUT						= 1 << 8,
-		EDF_MULTIPLE						= 1 << 9,
-		EDF_PRIMITIVE_DATA_TYPE			= 1 << 10,
-		EDF_SHARED						= 1 << 11,
-		EDF_NOINTERPOLATION				= 1 << 12,
+		EDF_CONST_ROW_MAJOR				= (1 << 0),
+		EDF_STATIC						= (1 << 1),
+		EDF_UNIFORM						= (1 << 2),
+		EDF_TEXTURE_SAMPLER_OR_BUFFER	= (1 << 3),
+		EDF_INITIALIZER					= (1 << 4),
+		EDF_INITIALIZER_LIST			= (1 << 5) | EDF_INITIALIZER,
+		EDF_SEMANTIC					= (1 << 6),
+		EDF_SEMICOLON					= (1 << 7),
+		EDF_IN_OUT						= (1 << 8),
+		EDF_MULTIPLE					= (1 << 9),
+		EDF_PRIMITIVE_DATA_TYPE			= (1 << 10),
+		EDF_SHARED						= (1 << 11),
+		EDF_INTERPOLATION				= (1 << 12),
 	};
 
 	EParseResult ParseInitializer(FHlslScanner& Scanner, FSymbolScope* SymbolScope, bool bAllowLists, FLinearAllocator* Allocator, AST::FExpression** OutList)
@@ -288,11 +281,15 @@ namespace CrossCompiler
 		return EParseResult::NotMatched;
 	}
 
-	EParseResult ParseDeclarationStorageQualifiers(FHlslScanner& Scanner, int32 Flags, bool& bOutPrimitiveFound, AST::FTypeQualifier* Qualifier)
+	EParseResult ParseDeclarationStorageQualifiers(FHlslScanner& Scanner, int32 TypeFlags, int32 DeclarationFlags, bool& bOutPrimitiveFound, AST::FTypeQualifier* Qualifier)
 	{
 		bOutPrimitiveFound = false;
 		int32 StaticFound = 0;
-		int32 NoInterpolationFound = 0;
+		int32 InterpolationLinearFound = 0;
+		int32 InterpolationCentroidFound = 0;
+		int32 InterpolationNoInterpolationFound = 0;
+		int32 InterpolationNoPerspectiveFound = 0;
+		int32 InterpolationSampleFound = 0;
 		int32 SharedFound = 0;
 		int32 ConstFound = 0;
 		int32 RowMajorFound = 0;
@@ -302,7 +299,7 @@ namespace CrossCompiler
 		int32 PrimitiveFound = 0;
 		int32 UniformFound = 0;
 
-		if (Flags & EDF_PRIMITIVE_DATA_TYPE)
+		if (DeclarationFlags & EDF_PRIMITIVE_DATA_TYPE)
 		{
 			const auto* Token = Scanner.GetCurrentToken();
 			if (Token && Token->Token == EHlslToken::Identifier)
@@ -322,7 +319,8 @@ namespace CrossCompiler
 		while (Scanner.HasMoreTokens())
 		{
 			bool bFound = false;
-			if ((Flags & EDF_STATIC) && Scanner.MatchToken(EHlslToken::Static))
+			auto* Token = Scanner.GetCurrentToken();
+			if ((DeclarationFlags & EDF_STATIC) && Scanner.MatchToken(EHlslToken::Static))
 			{
 				++StaticFound;
 				Qualifier->bIsStatic = true;
@@ -332,16 +330,7 @@ namespace CrossCompiler
 					return EParseResult::Error;
 				}
 			}
-			else if ((Flags & EDF_NOINTERPOLATION) && Scanner.MatchToken(EHlslToken::NoInterpolation))
-			{
-				++NoInterpolationFound;
-				if (NoInterpolationFound > 1)
-				{
-					Scanner.SourceError(TEXT("'nointerpolation' found more than once!\n"));
-					return EParseResult::Error;
-				}
-			}
-			else if ((Flags & EDF_SHARED) && Scanner.MatchToken(EHlslToken::GroupShared))
+			else if ((DeclarationFlags & EDF_SHARED) && Scanner.MatchToken(EHlslToken::GroupShared))
 			{
 				++SharedFound;
 				Qualifier->bShared = true;
@@ -351,7 +340,7 @@ namespace CrossCompiler
 					return EParseResult::Error;
 				}
 			}
-			else if ((Flags & EDF_CONST_ROW_MAJOR) && Scanner.MatchToken(EHlslToken::Const))
+			else if ((DeclarationFlags & EDF_CONST_ROW_MAJOR) && Scanner.MatchToken(EHlslToken::Const))
 			{
 				++ConstFound;
 				Qualifier->bConstant = true;
@@ -361,7 +350,7 @@ namespace CrossCompiler
 					return EParseResult::Error;
 				}
 			}
-			else if ((Flags & EDF_CONST_ROW_MAJOR) && Scanner.MatchToken(EHlslToken::RowMajor))
+			else if ((DeclarationFlags & EDF_CONST_ROW_MAJOR) && Scanner.MatchToken(EHlslToken::RowMajor))
 			{
 				++RowMajorFound;
 				Qualifier->bRowMajor = true;
@@ -371,7 +360,7 @@ namespace CrossCompiler
 					return EParseResult::Error;
 				}
 			}
-			else if ((Flags & EDF_IN_OUT) && Scanner.MatchToken(EHlslToken::In))
+			else if ((DeclarationFlags & EDF_IN_OUT) && Scanner.MatchToken(EHlslToken::In))
 			{
 				++InFound;
 				Qualifier->bIn = true;
@@ -386,7 +375,7 @@ namespace CrossCompiler
 					return EParseResult::Error;
 				}
 			}
-			else if ((Flags & EDF_IN_OUT) && Scanner.MatchToken(EHlslToken::Out))
+			else if ((DeclarationFlags & EDF_IN_OUT) && Scanner.MatchToken(EHlslToken::Out))
 			{
 				++OutFound;
 				Qualifier->bOut = true;
@@ -401,7 +390,7 @@ namespace CrossCompiler
 					return EParseResult::Error;
 				}
 			}
-			else if ((Flags & EDF_IN_OUT) && Scanner.MatchToken(EHlslToken::InOut))
+			else if ((DeclarationFlags & EDF_IN_OUT) && Scanner.MatchToken(EHlslToken::InOut))
 			{
 				++InOutFound;
 				Qualifier->bIn = true;
@@ -417,7 +406,7 @@ namespace CrossCompiler
 					return EParseResult::Error;
 				}
 			}
-			else if ((Flags & EDF_UNIFORM) && Scanner.MatchToken(EHlslToken::Uniform))
+			else if ((DeclarationFlags & EDF_UNIFORM) && Scanner.MatchToken(EHlslToken::Uniform))
 			{
 				++UniformFound;
 				Qualifier->bUniform = true;
@@ -427,13 +416,91 @@ namespace CrossCompiler
 					return EParseResult::Error;
 				}
 			}
+			else if ((DeclarationFlags & EDF_INTERPOLATION) && Token->Token == EHlslToken::Identifier)
+			{
+				if (Token->String == TEXT("linear"))
+				{
+					Scanner.Advance();
+					++InterpolationLinearFound;
+					Qualifier->bLinear = true;
+					if (InterpolationLinearFound > 1)
+					{
+						Scanner.SourceError(TEXT("'linear' found more than once!\n"));
+						return EParseResult::Error;
+					}
+				}
+				else if (Token->String == TEXT("centroid"))
+				{
+					Scanner.Advance();
+					++InterpolationCentroidFound;
+					Qualifier->bCentroid = true;
+					if (InterpolationCentroidFound > 1)
+					{
+						Scanner.SourceError(TEXT("'centroid' found more than once!\n"));
+						return EParseResult::Error;
+					}
+				}
+				else if (Token->String == TEXT("nointerpolation"))
+				{
+					Scanner.Advance();
+					++InterpolationNoInterpolationFound;
+					Qualifier->bNoInterpolation = true;
+					if (InterpolationNoInterpolationFound > 1)
+					{
+						Scanner.SourceError(TEXT("'nointerpolation' found more than once!\n"));
+						return EParseResult::Error;
+					}
+				}
+				else if (Token->String == TEXT("noperspective"))
+				{
+					Scanner.Advance();
+					++InterpolationNoPerspectiveFound;
+					Qualifier->bNoPerspective = true;
+					if (InterpolationNoPerspectiveFound > 1)
+					{
+						Scanner.SourceError(TEXT("'noperspective' found more than once!\n"));
+						return EParseResult::Error;
+					}
+				}
+				else if (Token->String == TEXT("sample"))
+				{
+					Scanner.Advance();
+					++InterpolationSampleFound;
+					Qualifier->bSample = true;
+					if (InterpolationSampleFound > 1)
+					{
+						Scanner.SourceError(TEXT("'sample' found more than once!\n"));
+						return EParseResult::Error;
+					}
+				}
+				else
+				{
+					break;
+				}
+			}
 			else
 			{
 				break;
 			}
 		}
 
-		if (UniformFound && (OutFound || InOutFound || PrimitiveFound || SharedFound || NoInterpolationFound))
+		int32 InterpolationFound = InterpolationLinearFound + InterpolationCentroidFound + InterpolationNoInterpolationFound + InterpolationNoPerspectiveFound + InterpolationSampleFound;
+		if (InterpolationFound)
+		{
+			if (InterpolationLinearFound && InterpolationNoInterpolationFound)
+			{
+				Scanner.SourceError(TEXT("Can't have both 'linear' and 'nointerpolation'!\n"));
+				return EParseResult::Error;
+			}
+
+			if (InterpolationCentroidFound && !(InterpolationLinearFound || InterpolationNoPerspectiveFound))
+			{
+				Scanner.SourceError(TEXT("'centroid' must be used with either 'linear' or 'noperspective'!\n"));
+				return EParseResult::Error;
+			}
+		}
+
+		if (UniformFound && (OutFound || InOutFound || PrimitiveFound || SharedFound || InterpolationFound))
 		{
 			Scanner.SourceError(TEXT("'uniform' can not be used with other storage qualifiers (inout, out, nointerpolation, etc)!\n"));
 			return EParseResult::Error;
@@ -441,17 +508,17 @@ namespace CrossCompiler
 
 		bOutPrimitiveFound = (PrimitiveFound > 0);
 
-		return (ConstFound + RowMajorFound + InFound + OutFound + InOutFound + StaticFound + SharedFound + PrimitiveFound + NoInterpolationFound + UniformFound)
+		return (ConstFound + RowMajorFound + InFound + OutFound + InOutFound + StaticFound + SharedFound + PrimitiveFound + InterpolationFound + UniformFound)
 			? EParseResult::Matched
 			: EParseResult::NotMatched;
 	}
 
-	EParseResult ParseGeneralDeclarationNoSemicolon(FHlslScanner& Scanner, FSymbolScope* SymbolScope, int32 Flags, FLinearAllocator* Allocator, AST::FDeclaratorList** OutDeclaratorList)
+	EParseResult ParseGeneralDeclarationNoSemicolon(FHlslScanner& Scanner, FSymbolScope* SymbolScope, int32 TypeFlags, int32 DeclarationFlags, FLinearAllocator* Allocator, AST::FDeclaratorList** OutDeclaratorList)
 	{
 		auto OriginalToken = Scanner.GetCurrentTokenIndex();
 		bool bPrimitiveFound = false;
 		auto* FullType = new(Allocator) AST::FFullySpecifiedType(Allocator, Scanner.GetCurrentToken()->SourceInfo);
-		auto ParseResult = ParseDeclarationStorageQualifiers(Scanner, Flags, bPrimitiveFound, &FullType->Qualifier);
+		auto ParseResult = ParseDeclarationStorageQualifiers(Scanner, TypeFlags, DeclarationFlags, bPrimitiveFound, &FullType->Qualifier);
 		if (ParseResult == EParseResult::Error)
 		{
 			return EParseResult::Error;
@@ -461,61 +528,85 @@ namespace CrossCompiler
 		auto* DeclaratorList = new(Allocator) AST::FDeclaratorList(Allocator, FullType->SourceInfo);
 		DeclaratorList->Type = FullType;
 
-		if (!bPrimitiveFound && (Flags & EDF_PRIMITIVE_DATA_TYPE))
+		if (!bPrimitiveFound && (DeclarationFlags & EDF_PRIMITIVE_DATA_TYPE))
 		{
 			const auto* StreamToken = Scanner.GetCurrentToken();
-			if (StreamToken && StreamToken->Token == EHlslToken::Identifier)
+			if (StreamToken)
 			{
-				if (StreamToken->String == TEXT("PointStream") ||
-					StreamToken->String == TEXT("LineStream") ||
-					StreamToken->String == TEXT("TriangleStream"))
+				if (StreamToken->Token == EHlslToken::Identifier)
 				{
-					Scanner.Advance();
-					bCanBeUnmatched = false;
-
-					if (!Scanner.MatchToken(EHlslToken::Lower))
+					if (StreamToken->String == TEXT("PointStream") ||
+						StreamToken->String == TEXT("LineStream") ||
+						StreamToken->String == TEXT("TriangleStream") ||
+						StreamToken->String == TEXT("InputPatch") ||
+						StreamToken->String == TEXT("OutputPatch"))
 					{
-						Scanner.SourceError(TEXT("Expected '<'!"));
-						return EParseResult::Error;
-					}
+						Scanner.Advance();
+						bCanBeUnmatched = false;
 
-					AST::FTypeSpecifier* TypeSpecifier = nullptr;
-					if (ParseGeneralType(Scanner, ETF_BUILTIN_NUMERIC | ETF_USER_TYPES, SymbolScope, Allocator, &TypeSpecifier) != EParseResult::Matched)
-					{
-						Scanner.SourceError(TEXT("Expected type!"));
-						return EParseResult::Error;
-					}
+						if (!Scanner.MatchToken(EHlslToken::Lower))
+						{
+							Scanner.SourceError(TEXT("Expected '<'!"));
+							return EParseResult::Error;
+						}
 
-					if (!Scanner.MatchToken(EHlslToken::Greater))
-					{
-						Scanner.SourceError(TEXT("Expected '>'!"));
-						return EParseResult::Error;
-					}
+						AST::FTypeSpecifier* TypeSpecifier = nullptr;
+						if (ParseGeneralType(Scanner, ETF_BUILTIN_NUMERIC | ETF_USER_TYPES, SymbolScope, Allocator, &TypeSpecifier) != EParseResult::Matched)
+						{
+							Scanner.SourceError(TEXT("Expected type!"));
+							return EParseResult::Error;
+						}
 
-					auto* IdentifierToken = Scanner.GetCurrentToken();
-					if (!Scanner.MatchToken(EHlslToken::Identifier))
-					{
-						Scanner.SourceError(TEXT("Expected identifier!"));
-						return EParseResult::Error;
-					}
+						if (StreamToken->String == TEXT("InputPatch") || StreamToken->String == TEXT("OutputPatch"))
+						{
+							if (!Scanner.MatchToken(EHlslToken::Comma))
+							{
+								Scanner.SourceError(TEXT("Expected ','!"));
+								return EParseResult::Error;
+							}
 
-					TypeSpecifier->InnerType = TypeSpecifier->TypeName;
-					TypeSpecifier->TypeName = Allocator->Strdup(StreamToken->String);
-					FullType->Specifier = TypeSpecifier;
+							//@todo-rco: Save this value!
+							auto* Elements = Scanner.GetCurrentToken();
+							if (!Scanner.MatchToken(EHlslToken::UnsignedIntegerConstant))
+							{
+								Scanner.SourceError(TEXT("Expected number!"));
+								return EParseResult::Error;
+							}
 
-					auto* Declaration = new(Allocator) AST::FDeclaration(Allocator, IdentifierToken->SourceInfo);
-					Declaration->Identifier = Allocator->Strdup(IdentifierToken->String);
+							TypeSpecifier->TextureMSNumSamples = Elements->UnsignedInteger;
+						}
 
-					DeclaratorList->Declarations.Add(Declaration);
-					*OutDeclaratorList = DeclaratorList;
-					return EParseResult::Matched;
+						if (!Scanner.MatchToken(EHlslToken::Greater))
+						{
+							Scanner.SourceError(TEXT("Expected '>'!"));
+							return EParseResult::Error;
+						}
+
+						auto* IdentifierToken = Scanner.GetCurrentToken();
+						if (!Scanner.MatchToken(EHlslToken::Identifier))
+						{
+							Scanner.SourceError(TEXT("Expected identifier!"));
+							return EParseResult::Error;
+						}
+
+						TypeSpecifier->InnerType = TypeSpecifier->TypeName;
+						TypeSpecifier->TypeName = Allocator->Strdup(StreamToken->String);
+						FullType->Specifier = TypeSpecifier;
+
+						auto* Declaration = new(Allocator)AST::FDeclaration(Allocator, IdentifierToken->SourceInfo);
+						Declaration->Identifier = Allocator->Strdup(IdentifierToken->String);
+
+						DeclaratorList->Declarations.Add(Declaration);
+						*OutDeclaratorList = DeclaratorList;
+						return EParseResult::Matched;
+					}
 				}
 			}
 		}
 
-		if (Flags & EDF_TEXTURE_SAMPLER_OR_BUFFER)
+		if (DeclarationFlags & EDF_TEXTURE_SAMPLER_OR_BUFFER)
 		{
-			auto Result = ParseTextureOrBufferSimpleDeclaration(Scanner, SymbolScope, (Flags & EDF_MULTIPLE) == EDF_MULTIPLE, Allocator, &DeclaratorList);
+			auto Result = ParseTextureOrBufferSimpleDeclaration(Scanner, SymbolScope, (DeclarationFlags & EDF_MULTIPLE) == EDF_MULTIPLE, Allocator, &DeclaratorList);
 			if (Result == EParseResult::Matched)
 			{
 				*OutDeclaratorList = DeclaratorList;
@@ -527,7 +618,7 @@ namespace CrossCompiler
 			}
 		}
 
-		const bool bAllowInitializerList = (Flags & EDF_INITIALIZER_LIST) == EDF_INITIALIZER_LIST;
+		const bool bAllowInitializerList = (DeclarationFlags & EDF_INITIALIZER_LIST) == EDF_INITIALIZER_LIST;
 
 		if (Scanner.MatchToken(EHlslToken::Struct))
 		{
@@ -552,7 +643,7 @@ namespace CrossCompiler
 						return EParseResult::Error;
 					}
 
-					if (Flags & EDF_INITIALIZER)
+					if (DeclarationFlags & EDF_INITIALIZER)
 					{
 						if (Scanner.MatchToken(EHlslToken::Equal))
 						{
@@ -567,7 +658,7 @@ namespace CrossCompiler
 					DeclaratorList->Declarations.Add(Declaration);
 				}
 			}
-			while ((Flags & EDF_MULTIPLE) == EDF_MULTIPLE && Scanner.MatchToken(EHlslToken::Comma));
+			while ((DeclarationFlags & EDF_MULTIPLE) == EDF_MULTIPLE && Scanner.MatchToken(EHlslToken::Comma));
 			*OutDeclaratorList = DeclaratorList;
 		}
 		else
@@ -595,7 +686,7 @@ namespace CrossCompiler
 					}
 
 					bool bSemanticFound = false;
-					if (Flags & EDF_SEMANTIC)
+					if (DeclarationFlags & EDF_SEMANTIC)
 					{
 						if (Scanner.MatchToken(EHlslToken::Colon))
 						{
@@ -611,7 +702,7 @@ namespace CrossCompiler
 						}
 					}
 					
-					if ((Flags & EDF_INITIALIZER) && !bSemanticFound)
+					if ((DeclarationFlags & EDF_INITIALIZER) && !bSemanticFound)
 					{
 						if (Scanner.MatchToken(EHlslToken::Equal))
 						{
@@ -625,7 +716,7 @@ namespace CrossCompiler
 
 					DeclaratorList->Declarations.Add(Declaration);
 				}
-				while ((Flags & EDF_MULTIPLE) == EDF_MULTIPLE && Scanner.MatchToken(EHlslToken::Comma));
+				while ((DeclarationFlags & EDF_MULTIPLE) == EDF_MULTIPLE && Scanner.MatchToken(EHlslToken::Comma));
 
 				*OutDeclaratorList = DeclaratorList;
 			}
@@ -639,17 +730,17 @@ namespace CrossCompiler
 		return EParseResult::Matched;
 	}
 
-	EParseResult ParseGeneralDeclaration(FHlslScanner& Scanner, FSymbolScope* SymbolScope, FLinearAllocator* Allocator, AST::FDeclaratorList** OutDeclaration, int32 Flags)
+	EParseResult ParseGeneralDeclaration(FHlslScanner& Scanner, FSymbolScope* SymbolScope, FLinearAllocator* Allocator, AST::FDeclaratorList** OutDeclaration, int32 TypeFlags, int32 DeclarationFlags)
 	{
 		auto OriginalToken = Scanner.GetCurrentTokenIndex();
 
-		auto Result = ParseGeneralDeclarationNoSemicolon(Scanner, SymbolScope, Flags, Allocator, OutDeclaration);
+		auto Result = ParseGeneralDeclarationNoSemicolon(Scanner, SymbolScope, TypeFlags, DeclarationFlags, Allocator, OutDeclaration);
 		if (Result == EParseResult::NotMatched || Result == EParseResult::Error)
 		{
 			return Result;
 		}
 
-		if (Flags & EDF_SEMICOLON)
+		if (DeclarationFlags & EDF_SEMICOLON)
 		{
 			if (!Scanner.MatchToken(EHlslToken::Semicolon))
 			{
@@ -693,7 +784,7 @@ namespace CrossCompiler
 				}
 
 				AST::FDeclaratorList* Declaration = nullptr;
-				auto Result = ParseGeneralDeclaration(Parser.Scanner, Parser.CurrentScope, Allocator, &Declaration, EDF_CONST_ROW_MAJOR | EDF_SEMICOLON | EDF_TEXTURE_SAMPLER_OR_BUFFER);
+				auto Result = ParseGeneralDeclaration(Parser.Scanner, Parser.CurrentScope, Allocator, &Declaration, 0, EDF_CONST_ROW_MAJOR | EDF_SEMICOLON | EDF_TEXTURE_SAMPLER_OR_BUFFER);
 				if (Result == EParseResult::Error)
 				{
 					return EParseResult::Error;
@@ -713,8 +804,15 @@ namespace CrossCompiler
 	EParseResult ParseStructBody(FHlslScanner& Scanner, FSymbolScope* SymbolScope, FLinearAllocator* Allocator, AST::FTypeSpecifier** OutTypeSpecifier)
 	{
 		const auto* Name = Scanner.GetCurrentToken();
-		if (Name && Scanner.MatchToken(EHlslToken::Identifier))
+		if (!Name)
+		{
+			return EParseResult::Error;
+		}
+
+		bool bAnonymous = true;
+		if (Scanner.MatchToken(EHlslToken::Identifier))
 		{
+			bAnonymous = false;
 			SymbolScope->Add(Name->String);
 		}
 
@@ -739,7 +837,8 @@ namespace CrossCompiler
 
 		auto* Struct = new(Allocator) AST::FStructSpecifier(Allocator, Name->SourceInfo);
 		Struct->ParentName = Allocator->Strdup(Parent);
-		Struct->Name = Allocator->Strdup(Name->String);
+		//@todo-rco: Differentiate anonymous!
+		Struct->Name = bAnonymous ? nullptr : Allocator->Strdup(Name->String);
 
 		bool bFoundRightBrace = false;
 		while (Scanner.HasMoreTokens())
@@ -751,7 +850,7 @@ namespace CrossCompiler
 			}
 
 			AST::FDeclaratorList* Declaration = nullptr;
-			auto Result = ParseGeneralDeclaration(Scanner, SymbolScope, Allocator, &Declaration, EDF_CONST_ROW_MAJOR | EDF_SEMICOLON | EDF_SEMANTIC | EDF_TEXTURE_SAMPLER_OR_BUFFER | EDF_NOINTERPOLATION);
+			auto Result = ParseGeneralDeclaration(Scanner, SymbolScope, Allocator, &Declaration, 0, EDF_CONST_ROW_MAJOR | EDF_SEMICOLON | EDF_SEMANTIC | EDF_TEXTURE_SAMPLER_OR_BUFFER | EDF_INTERPOLATION);
 			if (Result == EParseResult::Error)
 			{
 				return EParseResult::Error;
@@ -782,7 +881,7 @@ namespace CrossCompiler
 		while (Parser.Scanner.HasMoreTokens())
 		{
 			AST::FDeclaratorList* Declaration = nullptr;
-			auto Result = ParseGeneralDeclaration(Parser.Scanner, Parser.CurrentScope, Allocator, &Declaration, EDF_CONST_ROW_MAJOR | EDF_IN_OUT | EDF_TEXTURE_SAMPLER_OR_BUFFER | EDF_INITIALIZER | EDF_SEMANTIC | EDF_PRIMITIVE_DATA_TYPE | EDF_NOINTERPOLATION | EDF_UNIFORM);
+			auto Result = ParseGeneralDeclaration(Parser.Scanner, Parser.CurrentScope, Allocator, &Declaration, 0, EDF_CONST_ROW_MAJOR | EDF_IN_OUT | EDF_TEXTURE_SAMPLER_OR_BUFFER | EDF_INITIALIZER | EDF_SEMANTIC | EDF_PRIMITIVE_DATA_TYPE | EDF_INTERPOLATION | EDF_UNIFORM);
 			if (Result == EParseResult::Error)
 			{
 				return EParseResult::Error;
@@ -808,12 +907,16 @@ namespace CrossCompiler
 	{
 		auto OriginalToken = Parser.Scanner.GetCurrentTokenIndex();
 		AST::FTypeSpecifier* TypeSpecifier = nullptr;
-		auto Result = ParseGeneralType(Parser.Scanner, ETF_BUILTIN_NUMERIC | ETF_SAMPLER_TEXTURE_BUFFER | ETF_USER_TYPES | ETF_VOID, Parser.CurrentScope, Allocator, &TypeSpecifier);
+		auto Result = ParseGeneralType(Parser.Scanner, ETF_BUILTIN_NUMERIC | ETF_SAMPLER_TEXTURE_BUFFER | ETF_USER_TYPES | ETF_ERROR_IF_NOT_USER_TYPE | ETF_VOID, Parser.CurrentScope, Allocator, &TypeSpecifier);
 		if (Result == EParseResult::NotMatched)
 		{
 			Parser.Scanner.SetCurrentTokenIndex(OriginalToken);
 			return EParseResult::NotMatched;
 		}
+		else if (Result == EParseResult::Error)
+		{
+			return Result;
+		}
 
 		check(Result == EParseResult::Matched);
 
@@ -874,8 +977,7 @@ Done:
 			return EParseResult::NotMatched;
 		}
 
-		static FString Statement(TEXT("Statement"));
-		return TryRules(RulesStatements, Parser, Statement, false, Allocator, OutStatement);
+		return TryStatementRules(Parser, Allocator, OutStatement);
 	}
 
 	EParseResult ParseStatementBlock(FHlslParser& Parser, FLinearAllocator* Allocator, AST::FNode** OutStatement)
@@ -966,7 +1068,7 @@ check(0);
 	EParseResult ParseLocalDeclaration(FHlslParser& Parser, FLinearAllocator* Allocator, AST::FNode** OutDeclaration)
 	{
 		AST::FDeclaratorList* List = nullptr;
-		auto Result = ParseGeneralDeclaration(Parser.Scanner, Parser.CurrentScope, Allocator, &List, EDF_CONST_ROW_MAJOR | EDF_INITIALIZER | EDF_INITIALIZER_LIST | EDF_SEMICOLON | EDF_MULTIPLE);
+		auto Result = ParseGeneralDeclaration(Parser.Scanner, Parser.CurrentScope, Allocator, &List, 0, EDF_CONST_ROW_MAJOR | EDF_INITIALIZER | EDF_INITIALIZER_LIST | EDF_SEMICOLON | EDF_MULTIPLE | EDF_STATIC);
 		*OutDeclaration = List;
 		return Result;
 	}
@@ -974,7 +1076,7 @@ check(0);
 	EParseResult ParseGlobalVariableDeclaration(FHlslParser& Parser, FLinearAllocator* Allocator, AST::FNode** OutDeclaration)
 	{
 		AST::FDeclaratorList* List = nullptr;
-		auto Result = ParseGeneralDeclaration(Parser.Scanner, Parser.CurrentScope, Allocator, &List, EDF_CONST_ROW_MAJOR | EDF_STATIC | EDF_SHARED | EDF_TEXTURE_SAMPLER_OR_BUFFER | EDF_INITIALIZER | EDF_INITIALIZER_LIST | EDF_SEMICOLON | EDF_MULTIPLE | EDF_UNIFORM);
+		auto Result = ParseGeneralDeclaration(Parser.Scanner, Parser.CurrentScope, Allocator, &List, ETF_USER_TYPES | ETF_ERROR_IF_NOT_USER_TYPE, EDF_CONST_ROW_MAJOR | EDF_STATIC | EDF_SHARED | EDF_TEXTURE_SAMPLER_OR_BUFFER | EDF_INITIALIZER | EDF_INITIALIZER_LIST | EDF_SEMICOLON | EDF_MULTIPLE | EDF_UNIFORM | EDF_INTERPOLATION);
 		*OutDeclaration = List;
 		return Result;
 	}
@@ -1515,17 +1617,68 @@ check(0);
 		return EParseResult::Error;
 	}
 
+	EParseResult TryTranslationUnit(FHlslParser& Parser, FLinearAllocator* Allocator, AST::FNode** OutNode)
+	{
+		if (Parser.Scanner.MatchToken(EHlslToken::CBuffer))
+		{
+			auto Result = ParseCBuffer(Parser, Allocator, OutNode);
+			if (Result == EParseResult::Error || Result == EParseResult::Matched)
+			{
+				return Result;
+			}
+		}
+
+		// Match Attributes
+		TLinearArray<AST::FAttribute*> Attributes(Allocator);
+		while (Parser.Scanner.HasMoreTokens())
+		{
+			const auto* Peek = Parser.Scanner.GetCurrentToken();
+			if (Peek->Token == EHlslToken::LeftSquareBracket)
+			{
+				AST::FAttribute* Attribute = nullptr;
+				auto Result = TryParseAttribute(Parser, Allocator, &Attribute);
+				if (Result == EParseResult::Matched)
+				{
+					Attributes.Add(Attribute);
+					continue;
+				}
+				else if (Result == EParseResult::Error)
+				{
+					return EParseResult::Error;
+				}
+			}
+
+			break;
+		}
+
+		const auto* Peek = Parser.Scanner.GetCurrentToken();
+		if (!Peek)
+		{
+			return EParseResult::Error;
+		}
+		
+		auto Result = ParseFunctionDeclaration(Parser, Allocator, OutNode);
+		if (Result == EParseResult::Error || Result == EParseResult::Matched)
+		{
+			return Result;
+		}
+
+		Result = ParseGlobalVariableDeclaration(Parser, Allocator, OutNode);
+		if (Result == EParseResult::Error || Result == EParseResult::Matched)
+		{
+			return Result;
+		}
+
+		Parser.Scanner.SourceError(TEXT("Unable to match rule!"));
+		return EParseResult::Error;
+	}
+
 	namespace ParserRules
 	{
 		static struct FStaticInitializer
 		{
 			FStaticInitializer()
 			{
-				// Top Level constructs
-				RulesTranslationUnit.Add(FRulePair(EHlslToken::CBuffer, ParseCBuffer));
-				RulesTranslationUnit.Add(FRulePair(EHlslToken::Invalid, ParseFunctionDeclaration, true));
-				RulesTranslationUnit.Add(FRulePair(EHlslToken::Invalid, ParseGlobalVariableDeclaration));
-
 				RulesStatements.Add(FRulePair(EHlslToken::LeftBrace, ParseStatementBlock));
 				RulesStatements.Add(FRulePair(EHlslToken::Return, ParseReturnStatement));
 				RulesStatements.Add(FRulePair(EHlslToken::Do, ParseDoStatement));
@@ -1544,6 +1697,7 @@ check(0);
 	}
 
 	FHlslParser::FHlslParser(FLinearAllocator* InAllocator) :
+		Scanner(),
 		GlobalScope(InAllocator, nullptr),
 		Allocator(InAllocator)
 	{
@@ -1569,7 +1723,7 @@ check(0);
 
 				static FString GlobalDeclOrDefinition(TEXT("Global declaration or definition"));
 				AST::FNode* Node = nullptr;
-				auto Result = TryRules(RulesTranslationUnit, Parser, GlobalDeclOrDefinition, true, &Allocator, &Node);
+				auto Result = TryTranslationUnit(Parser, &Allocator, &Node);
 				if (Result == EParseResult::Error)
 				{
 					bSuccess = false;
diff --git a/Engine/Source/Developer/ShaderCompilerCommon/Private/HlslUtils.cpp b/Engine/Source/Developer/ShaderCompilerCommon/Private/HlslUtils.cpp
new file mode 100644
--- /dev/null
+++ b/Engine/Source/Developer/ShaderCompilerCommon/Private/HlslUtils.cpp
@@ -0,0 +1,75 @@
+// Copyright 1998-2015 Epic Games, Inc. All Rights Reserved.
+
+/*=============================================================================
+	HlslUtils.cpp - Utils for HLSL.
+=============================================================================*/
+
+#include "ShaderCompilerCommon.h"
+#include "HlslUtils.h"
+
+namespace CrossCompiler
+{
+	namespace Memory
+	{
+		static struct FPagePoolInstance
+		{
+			~FPagePoolInstance()
+			{
+				check(UsedPages.Num() == 0);
+				for (int32 Index = 0; Index < FreePages.Num(); ++Index)
+				{
+					delete FreePages[Index];
+				}
+			}
+
+			FPage* AllocatePage()
+			{
+				FScopeLock ScopeLock(&CriticalSection);
+
+				if (FreePages.Num() == 0)
+				{
+					FreePages.Add(new FPage(PageSize));
+				}
+
+				auto* Page = FreePages.Last();
+				FreePages.RemoveAt(FreePages.Num() - 1, 1, false);
+				UsedPages.Add(Page);
+				return Page;
+			}
+
+			void FreePage(FPage* Page)
+			{
+				FScopeLock ScopeLock(&CriticalSection);
+
+				int32 Index = UsedPages.Find(Page);
+				check(Index >= 0);
+				UsedPages.RemoveAt(Index, 1, false);
+				FreePages.Add(Page);
+			}
+
+			TArray<FPage*, TInlineAllocator<8> > FreePages;
+			TArray<FPage*, TInlineAllocator<8> > UsedPages;
+
+			FCriticalSection CriticalSection;
+
+		} GMemoryPagePool;
+
+		FPage* FPage::AllocatePage()
+		{
+#if USE_PAGE_POOLING
+			return GMemoryPagePool.AllocatePage();
+#else
+			return new FPage(PageSize);
+#endif
+		}
+
+		void FPage::FreePage(FPage* Page)
+		{
+#if USE_PAGE_POOLING
+			GMemoryPagePool.FreePage(Page);
+#else
+			delete Page;
+#endif
+		}
+	}
+}
diff --git a/Engine/Source/Developer/ShaderCompilerCommon/Private/HlslUtils.h b/Engine/Source/Developer/ShaderCompilerCommon/Private/HlslUtils.h
new file mode 100644
--- /dev/null
+++ b/Engine/Source/Developer/ShaderCompilerCommon/Private/HlslUtils.h
@@ -0,0 +1,233 @@
+// Copyright 1998-2015 Epic Games, Inc. All Rights Reserved.
+
+/*=============================================================================
+	HlslUtils.h - Utilities for Hlsl.
+=============================================================================*/
+
+#pragma once
+
+#define USE_UNREAL_ALLOCATOR	0
+#define USE_PAGE_POOLING		0
+
+namespace CrossCompiler
+{
+	namespace Memory
+	{
+		struct FPage
+		{
+			int8* Current;
+			int8* Begin;
+			int8* End;
+
+			FPage(SIZE_T Size)
+			{
+				check(Size > 0);
+				Begin = new int8[Size];
+				End = Begin + Size;
+				Current = Begin;
+			}
+
+			~FPage()
+			{
+				delete[] Begin;
+			}
+
+			static FPage* AllocatePage();
+			static void FreePage(FPage* Page);
+		};
+
+		enum
+		{
+			PageSize = 64 * 1024
+		};
+	}
+
+	struct FLinearAllocator
+	{
+		FLinearAllocator()
+		{
+			auto* Initial = Memory::FPage::AllocatePage();
+			Pages.Add(Initial);
+		}
+
+		~FLinearAllocator()
+		{
+			for (auto* Page : Pages)
+			{
+				Memory::FPage::FreePage(Page);
+			}
+		}
+
+		inline void* Alloc(SIZE_T NumBytes)
+		{
+			check(NumBytes <= Memory::PageSize);
+			auto* Page = Pages.Last();
+			if (Page->Current + NumBytes > Page->End)
+			{
+				Page = Memory::FPage::AllocatePage();
+				Pages.Add(Page);
+			}
+
+			void* Ptr = Page->Current;
+			Page->Current += NumBytes;
+			return Ptr;
+		}
+
+		inline void* Alloc(SIZE_T NumBytes, SIZE_T Align)
+		{
+			void* Data = Alloc(NumBytes + Align - 1);
+			UPTRINT Address = (UPTRINT)Data;
+			Address += (Align - (Address % (UPTRINT)Align)) % Align;
+			return (void*)Address;
+		}
+
+		/*inline*/ TCHAR* Strdup(const TCHAR* String)
+		{
+			if (!String)
+			{
+				return nullptr;
+			}
+
+			const auto Length = FCString::Strlen(String);
+			const auto Size = (Length + 1) * sizeof(TCHAR);
+			auto* Data = (TCHAR*)Alloc(Size);
+			FCString::Strcpy(Data, Length, String);
+			return Data;
+		}
+
+		inline TCHAR* Strdup(const FString& String)
+		{
+			return Strdup(*String);
+		}
+
+		TArray<Memory::FPage*, TInlineAllocator<8> > Pages;
+	};
+
+#if !USE_UNREAL_ALLOCATOR
+	class FLinearAllocatorPolicy
+	{
+	public:
+		// Unreal allocator magic
+		enum { NeedsElementType = false };
+		enum { RequireRangeCheck = true };
+
+		template<typename ElementType>
+		class ForElementType
+		{
+		public:
+
+			/** Default constructor. */
+			ForElementType() :
+				LinearAllocator(nullptr),
+				Data(nullptr)
+			{}
+
+			// FContainerAllocatorInterface
+			/*FORCEINLINE*/ ElementType* GetAllocation() const
+			{
+				return Data;
+			}
+			void ResizeAllocation(int32 PreviousNumElements, int32 NumElements, int32 NumBytesPerElement)
+			{
+				void* OldData = Data;
+				if (NumElements)
+				{
+					// Allocate memory from the stack.
+					Data = (ElementType*)LinearAllocator->Alloc(NumElements * NumBytesPerElement,
+						FMath::Max((uint32)sizeof(void*), (uint32)ALIGNOF(ElementType))
+						);
+
+					// If the container previously held elements, copy them into the new allocation.
+					if (OldData && PreviousNumElements)
+					{
+						const int32 NumCopiedElements = FMath::Min(NumElements, PreviousNumElements);
+						FMemory::Memcpy(Data, OldData, NumCopiedElements * NumBytesPerElement);
+					}
+				}
+			}
+			int32 CalculateSlack(int32 NumElements, int32 NumAllocatedElements, int32 NumBytesPerElement) const
+			{
+				return DefaultCalculateSlack(NumElements, NumAllocatedElements, NumBytesPerElement);
+			}
+
+			int32 GetAllocatedSize(int32 NumAllocatedElements, int32 NumBytesPerElement) const
+			{
+				return NumAllocatedElements * NumBytesPerElement;
+			}
+
+			FLinearAllocator* LinearAllocator;
+
+		private:
+
+			/** A pointer to the container's elements. */
+			ElementType* Data;
+		};
+
+		typedef ForElementType<FScriptContainerElement> ForAnyElementType;
+	};
+
+	class FLinearBitArrayAllocator
+		: public TInlineAllocator<4, FLinearAllocatorPolicy>
+	{
+	};
+
+	class FLinearSparseArrayAllocator
+		: public TSparseArrayAllocator<FLinearAllocatorPolicy, FLinearBitArrayAllocator>
+	{
+	};
+
+	class FLinearSetAllocator
+		: public TSetAllocator<FLinearSparseArrayAllocator, TInlineAllocator<1, FLinearAllocatorPolicy> >
+	{
+	};
+
+	template <typename TType>
+	class TLinearArray : public TArray<TType, FLinearAllocatorPolicy>
+	{
+	public:
+		TLinearArray(FLinearAllocator* Allocator)
+		{
+			TArray<TType, FLinearAllocatorPolicy>::AllocatorInstance.LinearAllocator = Allocator;
+		}
+	};
+
+	/*
+	template <typename TType>
+	struct TLinearSet : public TSet<typename TType, DefaultKeyFuncs<typename TType>, FLinearSetAllocator>
+	{
+	TLinearSet(FLinearAllocator* InAllocator)
+	{
+	Elements.AllocatorInstance.LinearAllocator = InAllocator;
+	}
+	};*/
+#endif
+
+	struct FSourceInfo
+	{
+		FString* Filename;
+		int32 Line;
+		int32 Column;
+
+		FSourceInfo() : Filename(nullptr), Line(0), Column(0) {}
+	};
+
+	inline void SourceError(const FSourceInfo& SourceInfo, const TCHAR* String)
+	{
+		FPlatformMisc::LowLevelOutputDebugStringf(TEXT("%s(%d): (%d) %s\n"), *SourceInfo.Filename, SourceInfo.Line, SourceInfo.Column, String);
+	}
+
+	inline void SourceError(const TCHAR* String)
+	{
+		FPlatformMisc::LowLevelOutputDebugStringf(TEXT("%s\n"), String);
+	}
+
+	inline void SourceWarning(const FSourceInfo& SourceInfo, const TCHAR* String)
+	{
+		FPlatformMisc::LowLevelOutputDebugStringf(TEXT("%s(%d): (%d) %s\n"), *SourceInfo.Filename, SourceInfo.Line, SourceInfo.Column, String);
+	}
+
+	inline void SourceWarning(const TCHAR* String)
+	{
+		FPlatformMisc::LowLevelOutputDebugStringf(TEXT("%s\n"), String);
+	}
+}
