CCT - Compile fix

diff --git a/Engine/Source/Developer/ShaderCompilerCommon/Private/HlslUtils.h b/Engine/Source/Developer/ShaderCompilerCommon/Private/HlslUtils.h
--- a/Engine/Source/Developer/ShaderCompilerCommon/Private/HlslUtils.h
+++ b/Engine/Source/Developer/ShaderCompilerCommon/Private/HlslUtils.h
@@ -213,7 +213,7 @@ namespace CrossCompiler
 
 	inline void SourceError(const FSourceInfo& SourceInfo, const TCHAR* String)
 	{
-		FPlatformMisc::LowLevelOutputDebugStringf(TEXT("%s(%d): (%d) %s\n"), *SourceInfo.Filename, SourceInfo.Line, SourceInfo.Column, String);
+		FPlatformMisc::LowLevelOutputDebugStringf(TEXT("%s(%d): (%d) %s\n"), **SourceInfo.Filename, SourceInfo.Line, SourceInfo.Column, String);
 	}
 
 	inline void SourceError(const TCHAR* String)
@@ -223,7 +223,7 @@ namespace CrossCompiler
 
 	inline void SourceWarning(const FSourceInfo& SourceInfo, const TCHAR* String)
 	{
-		FPlatformMisc::LowLevelOutputDebugStringf(TEXT("%s(%d): (%d) %s\n"), *SourceInfo.Filename, SourceInfo.Line, SourceInfo.Column, String);
+		FPlatformMisc::LowLevelOutputDebugStringf(TEXT("%s(%d): (%d) %s\n"), **SourceInfo.Filename, SourceInfo.Line, SourceInfo.Column, String);
 	}
 
 	inline void SourceWarning(const TCHAR* String)
