PVS-Studio detected that if (BufferIndex != FRHIResourceTableEntry::GetEndOfStreamToken()) will always return true because BufferIndex is 16bit and GetEndOfStreamToken returns a 32bit constant that is out of range of a 16bit integer.

diff --git a/Engine/Source/Developer/ShaderCompilerCommon/Private/ShaderCompilerCommon.cpp b/Engine/Source/Developer/ShaderCompilerCommon/Private/ShaderCompilerCommon.cpp
--- a/Engine/Source/Developer/ShaderCompilerCommon/Private/ShaderCompilerCommon.cpp
+++ b/Engine/Source/Developer/ShaderCompilerCommon/Private/ShaderCompilerCommon.cpp
@@ -16,7 +16,7 @@ int16 GetNumUniformBuffersUsed(const FShaderResourceTable& InSRT)
 						for (int32 i = 0; i < In.Num(); ++i)
 						{
 							auto BufferIndex = FRHIResourceTableEntry::GetUniformBufferIndex(In[i]);
-							if (BufferIndex != FRHIResourceTableEntry::GetEndOfStreamToken())
+							if (BufferIndex != static_cast<uint16>(FRHIResourceTableEntry::GetEndOfStreamToken()) )
 							{
 								LastIndex = FMath::Max(LastIndex, (int16)BufferIndex);
 							}
