UE4 - Fix aggresive search & replace for uniform buffer prefixes

diff --git a/Engine/Source/Developer/ShaderCompilerCommon/Private/ShaderCompilerCommon.cpp b/Engine/Source/Developer/ShaderCompilerCommon/Private/ShaderCompilerCommon.cpp
--- a/Engine/Source/Developer/ShaderCompilerCommon/Private/ShaderCompilerCommon.cpp
+++ b/Engine/Source/Developer/ShaderCompilerCommon/Private/ShaderCompilerCommon.cpp
@@ -128,7 +128,7 @@ void BuildResourceTableMapping(
 	OutSRT.MaxBoundResourceTable = MaxBoundResourceTable;
 }
 
-// Specialized version of FString::ReplaceInline that checks that the search text is not inside a #line directive
+// Specialized version of FString::ReplaceInline that checks that the search word is not inside a #line directive
 static void WholeWordReplaceInline(FString& String, TCHAR* StartPtr, const TCHAR* SearchText, const TCHAR* ReplacementText)
 {
 	if (String.Len() > 0
@@ -184,6 +184,19 @@ static void WholeWordReplaceInline(FString& String, TCHAR* StartPtr, const TCHAR
 				}
 			}
 			
+			// Make sure this is not part of an identifier
+			if (bReplace && Pos > StartPtr)
+			{
+				const auto Char = Pos[-1];
+				if ((Char >= 'a' && Char <= 'z') ||
+					(Char >= 'A' && Char <= 'Z') ||
+					(Char >= '0' && Char <= '9') ||
+					Char == '_')
+				{
+					bReplace = false;
+				}
+			}
+
 			if (bReplace)
 			{
 				// FCString::Strcpy inserts a terminating zero so can't use that
