UE4 - Fix Metal

diff --git a/Engine/Source/Developer/ShaderCompilerCommon/Private/ShaderCompilerCommon.cpp b/Engine/Source/Developer/ShaderCompilerCommon/Private/ShaderCompilerCommon.cpp
--- a/Engine/Source/Developer/ShaderCompilerCommon/Private/ShaderCompilerCommon.cpp
+++ b/Engine/Source/Developer/ShaderCompilerCommon/Private/ShaderCompilerCommon.cpp
@@ -1068,14 +1068,12 @@ namespace CrossCompiler
 				Attribute.ArrayCount = 0;
 			}
 
-			if (!Match(ShaderSource, ';'))
+			if (Match(ShaderSource, ';'))
 			{
-				return false;
-			}
-
-			if (!ParseSignedNumber(ShaderSource, Attribute.Index))
-			{
-				return false;
+				if (!ParseSignedNumber(ShaderSource, Attribute.Index))
+				{
+					return false;
+				}
 			}
 
 			if (!Match(ShaderSource, ':'))
diff --git a/Engine/Source/Developer/iOS/MetalShaderFormat/Private/MetalShaderFormat.cpp b/Engine/Source/Developer/iOS/MetalShaderFormat/Private/MetalShaderFormat.cpp
--- a/Engine/Source/Developer/iOS/MetalShaderFormat/Private/MetalShaderFormat.cpp
+++ b/Engine/Source/Developer/iOS/MetalShaderFormat/Private/MetalShaderFormat.cpp
@@ -16,7 +16,7 @@ class FMetalShaderFormat : public IShaderFormat
 {
 	enum
 	{
-		HEADER_VERSION = 18,
+		HEADER_VERSION = 19,
 	};
 	struct FVersion
 	{
diff --git a/Engine/Source/Programs/CrossCompilerTool/Private/CrossCompilerTool.cpp b/Engine/Source/Programs/CrossCompilerTool/Private/CrossCompilerTool.cpp
--- a/Engine/Source/Programs/CrossCompilerTool/Private/CrossCompilerTool.cpp
+++ b/Engine/Source/Programs/CrossCompilerTool/Private/CrossCompilerTool.cpp
@@ -170,6 +170,7 @@ namespace CCT
 			}
 		}
 
+		if (ShaderSource)
 		{
 			const ANSICHAR* USFSource = ShaderSource;
 			CrossCompiler::FHlslccHeader CCHeader;
