UE4 - Metal shader compiler fixes

diff --git a/Engine/Shaders/MetalCommon.usf b/Engine/Shaders/MetalCommon.usf
--- a/Engine/Shaders/MetalCommon.usf
+++ b/Engine/Shaders/MetalCommon.usf
@@ -5,6 +5,6 @@
 =============================================================================*/
 
 // Update this GUID to improve shader recompilation for Metal only shaders
-// GUID = {06C79EC9-644E-4D22-ADA9-47DD4BB343D9}
+// GUID = {06C79EC9-644E-4D22-ADA9-47DD4BB343D9a}
   
 #pragma once
diff --git a/Engine/Source/Developer/ShaderCompilerCommon/Private/ShaderCompilerCommon.cpp b/Engine/Source/Developer/ShaderCompilerCommon/Private/ShaderCompilerCommon.cpp
--- a/Engine/Source/Developer/ShaderCompilerCommon/Private/ShaderCompilerCommon.cpp
+++ b/Engine/Source/Developer/ShaderCompilerCommon/Private/ShaderCompilerCommon.cpp
@@ -445,16 +445,30 @@ namespace CrossCompiler
 		}
 
 		// Read shader name if any
-		ANSICHAR const* ShaderName = nullptr;
-		uint32 ShaderNameLen = 0;
-		while (FCStringAnsi::Strncmp(ShaderSource, "// !", 4) == 0)
+		if (FCStringAnsi::Strncmp(ShaderSource, "// !", 4) == 0)
 		{
-			ShaderName = ShaderSource;
+			ShaderSource += 4;
 			while (*ShaderSource && *ShaderSource != '\n')
 			{
 				Name += (TCHAR)*ShaderSource;
 				++ShaderSource;
 			}
+
+			if (*ShaderSource == '\n')
+			{
+				++ShaderSource;
+			}
+		}
+
+		// Skip any comments that come before the signature.
+		while (FCStringAnsi::Strncmp(ShaderSource, "//", 2) == 0 &&
+			FCStringAnsi::Strncmp(ShaderSource + 2, " @", 2) != 0)
+		{
+			ShaderSource += 2;
+			while (*ShaderSource && *ShaderSource++ != '\n')
+			{
+				// Do nothing
+			}
 		}
 
 		if (FCStringAnsi::Strncmp(ShaderSource, InputsPrefix, InputsPrefixLen) == 0)
@@ -907,6 +921,22 @@ namespace CrossCompiler
 				}
 
 				SamplerStates.Add(SamplerState);
+
+				// Break if EOL
+				if (Match(ShaderSource, '\n'))
+				{
+					break;
+				}
+
+				// Has to be a comma!
+				if (Match(ShaderSource, ','))
+				{
+					continue;
+				}
+
+				//#todo-rco: Need a log here
+				//UE_LOG(ShaderCompilerCommon, Warning, TEXT("Invalid char '%c'"), *ShaderSource);
+				return false;
 			}
 		}
 
diff --git a/Engine/Source/Programs/CrossCompilerTool/Private/CrossCompilerTool.cpp b/Engine/Source/Programs/CrossCompilerTool/Private/CrossCompilerTool.cpp
--- a/Engine/Source/Programs/CrossCompilerTool/Private/CrossCompilerTool.cpp
+++ b/Engine/Source/Programs/CrossCompilerTool/Private/CrossCompilerTool.cpp
@@ -180,7 +180,7 @@ namespace CCT
 				UE_LOG(LogCrossCompilerTool, Error, TEXT("Bad hlslcc header found"));
 			}
 
-			if (*USFSource != '#')
+			if (Language == &GlslLanguage && *USFSource != '#')
 			{
 				UE_LOG(LogCrossCompilerTool, Error, TEXT("Bad hlslcc header found! Missing '#'!"));
 			}
