UE4 - hlslcc - Fix tessellation on OpenGL4

diff --git a/Engine/Source/Developer/ShaderCompilerCommon/Private/ShaderCompilerCommon.cpp b/Engine/Source/Developer/ShaderCompilerCommon/Private/ShaderCompilerCommon.cpp
--- a/Engine/Source/Developer/ShaderCompilerCommon/Private/ShaderCompilerCommon.cpp
+++ b/Engine/Source/Developer/ShaderCompilerCommon/Private/ShaderCompilerCommon.cpp
@@ -1116,6 +1116,21 @@ namespace CrossCompiler
 				return false;
 			}
 
+			// Optional array suffix
+			if (Match(ShaderSource, '['))
+			{
+				Attribute.Name += '[';
+				while (*ShaderSource)
+				{
+					Attribute.Name += *ShaderSource;
+					if (Match(ShaderSource, ']'))
+					{
+						break;
+					}
+					++ShaderSource;
+				}
+			}
+
 			OutAttributes.Add(Attribute);
 
 			// Break if EOL
