Copying //UE4/Dev-Rendering to Dev-Main (Source //UE4/Dev-Rendering@2932636)

diff --git a/Engine/Shaders/TexCoordScaleAnalysisPixelShader.usf b/Engine/Shaders/TexCoordScaleAnalysisPixelShader.usf
--- a/Engine/Shaders/TexCoordScaleAnalysisPixelShader.usf
+++ b/Engine/Shaders/TexCoordScaleAnalysisPixelShader.usf
@@ -6,66 +6,78 @@
 
 #define TEX_COORD_SCALE_ANALYSIS 1
 #define DEBUG_MATERIAL_PARAMETERS 1
-#define INITIAL_GPU_SCALE 1024
 
 #include "Common.usf"
 #include "SHCommon.usf"
 
-// Scale is 16 bits interger, 16 bits fractionnal.
-const static uint FRACTIONNAL_SCALE = 0xFFFF;
-
-uint FloatToDiscreteScale(float Value)
-{
-	return (uint)(Value * (float)FRACTIONNAL_SCALE); 
-}
-
-RWTexture2D<uint> RWCoordScales : register(u7);
-
  // If -1, output everything to the RW texture.
 float4 AccuracyColors[5];
 uint TextureAnalysisIndex;
-float CPUScale;
+
+float4 OneOverCPUTexCoordScales[MAX_NUM_TEXTURE_REGISTER / 4];
+int4 TexCoordIndices[MAX_NUM_TEXTURE_REGISTER / 4];
+
+float GetComponent(float4 V, int Index)
+{
+	FLATTEN
+	if (Index == 0) return V.x;
+	FLATTEN
+	if (Index == 1) return V.y;
+	FLATTEN
+	if (Index == 2) return V.z;
+	return V.w;
+}
 
 // Global var used to reduce the number of param passed to the shaders.
 uint2 PixelPosition;
-uint CoordScaleHeight;
-uint CoordScaleWidth;
-float OneOverDDX0;
-float OneOverDDY0;
-float GPUScale;
+
+// ddU, ddV for each coord index (up to 4 currently)
+float4 OneOverDDU;
+float4 OneOverDDV;
+
+// Used when inspecting single elements.
+float MinScale;
+float MaxScale;
+
 float TexSample;
-  
-MaterialFloat StoreTexCoordScale(float2 UV, uint TextureReferenceIndex)
+
+// When computing the texture scales, this holds 4 scales for consecutive texture register indices.
+float4 ScalesPerIndex;
+
+MaterialFloat StoreTexCoordScale(float2 UV, uniform uint TextureReferenceIndex)
 {
- 	// Take the minimum scale since high scale requires less resolution.
-	float Scale = min(length(ddx(UV)) * OneOverDDX0, length(ddy(UV)) * OneOverDDY0);
+	// Take the minimum scale since high scale requires less resolution.
+	float GPUScaleX = length(ddx(UV));
+	float GPUScaleY = length(ddy(UV));
 
-	if (TextureAnalysisIndex == TextureReferenceIndex)
+	if (TextureReferenceIndex >= 0 && TextureReferenceIndex < MAX_NUM_TEXTURE_REGISTER)
 	{
-		GPUScale = min(Scale, GPUScale);
-	} 
+		float OneOverCPUScale = OneOverCPUTexCoordScales[TextureReferenceIndex / 4][TextureReferenceIndex % 4];
 
-	[branch]
-	if (TextureAnalysisIndex == -1) 
-	{
-		InterlockedMin(RWCoordScales[PixelPosition + uint2(TextureReferenceIndex * CoordScaleWidth, 0)], FloatToDiscreteScale(Scale));
-	}
+		int TexCoordIndex = TexCoordIndices[TextureReferenceIndex / 4][TextureReferenceIndex % 4];
+		
+		float GPUScale = min(GPUScaleX * GetComponent(OneOverDDU, TexCoordIndex), GPUScaleY * GetComponent(OneOverDDV, TexCoordIndex));
 
+		// Code path when searching the worst scales
+		const bool bUpdateMinMax = (OneOverCPUScale > 0 && (TextureAnalysisIndex == MAX_NUM_TEXTURE_REGISTER || TextureAnalysisIndex == TextureReferenceIndex));
+		MinScale = bUpdateMinMax ? min(MinScale, GPUScale * OneOverCPUScale) : MinScale;
+		MaxScale = bUpdateMinMax ? max(MaxScale, GPUScale * OneOverCPUScale) : MaxScale;
+
+		// Code path when rendering the material in a render target. Each TILE_RESOLUTION X TILE_RESOLUTION maps to different texture.
+		const bool bUpdateScale = (TextureAnalysisIndex == -1 && PixelPosition.x / TILE_RESOLUTION == TextureReferenceIndex / 4) ;
+		ScalesPerIndex[TextureReferenceIndex % 4] = bUpdateScale ? min(ScalesPerIndex[TextureReferenceIndex % 4], GPUScale) : ScalesPerIndex[TextureReferenceIndex % 4];
+	}
 	return 1.f;
 }
 
-MaterialFloat StoreTexSample(float4 C, uint TextureReferenceIndex)
+MaterialFloat StoreTexSample(float4 C, uniform uint TextureReferenceIndex)
 {
-	if (TextureAnalysisIndex == TextureReferenceIndex)
-	{
-		// Alpha?
-		TexSample = lerp(.25f, 1.f, Luminance(C.rgb));
-	}
+	// Alpha?
+	TexSample = TextureAnalysisIndex == TextureReferenceIndex ? lerp(.4f, 1.f, saturate(Luminance(C.rgb))) : TexSample;
 
 	return 1.f;
 }
 
-
 #include "Material.usf"
 #include "DebugViewModeCommon.usf"
 
@@ -81,18 +93,53 @@ void Main(
 	ResolvedView = ResolveView();
 #endif
 
-
-	float3 Result = float3(.1, .1, .1);
+	// This default value will make it dark grey when nothing updates the texSample.
+	float3 Result = float3(.02, .02, .02);
 
 	// Set global var used for outputting data
 	PixelPosition = DebugInputs.SvPosition;
-	RWCoordScales.GetDimensions(CoordScaleWidth, CoordScaleHeight);
-	OneOverDDX0 = 1 / length(ddx(DebugInputs.TexCoord0));
-	OneOverDDY0 = 1 / length(ddy(DebugInputs.TexCoord0));
-	GPUScale = INITIAL_GPU_SCALE;
+	OneOverDDU = 1 / float4(length(ddx(DebugInputs.TexCoord01.xy)), length(ddx(DebugInputs.TexCoord01.zw)), length(ddx(DebugInputs.TexCoord23.xy)), length(ddx(DebugInputs.TexCoord23.zw)));
+	OneOverDDV = 1 / float4(length(ddy(DebugInputs.TexCoord01.xy)), length(ddy(DebugInputs.TexCoord01.zw)), length(ddy(DebugInputs.TexCoord23.xy)), length(ddy(DebugInputs.TexCoord23.zw)));
+
 	TexSample = 1.f;
+	MinScale = INITIAL_GPU_SCALE;
+	MaxScale = 0;
+	ScalesPerIndex = INITIAL_GPU_SCALE;
 
 	FMaterialPixelParameters MaterialParameters = GetMaterialPixelParameters(DebugInputs, DebugInputs.SvPosition);
+
+	// This is used to patch invalid parameters when computing the scales with FMeshMaterialRenderItem::EnqueueMaterialRender
+	FLATTEN
+	if (TextureAnalysisIndex == -1)
+	{
+		// @TODO : The quad vertex data has invalid tangent space for some reason, even though the code in FMeshMaterialRenderItem::FillQuadData looks like setting it right.
+		MaterialParameters.TangentToWorld = float3x3(1,0,0,0,1,0,0,0,1);
+
+		// The next line ensures a constant view vector of (0,0,1) for all pixels. Required because here SVPositionToTranslatedWorld is identity, making excessive view angle increase per pixel.
+		// That creates bad side effects for anything that depends on the view vector, like parallax or bump offset mappings. For those, we want the tangent
+		// space view vector to be perpendicular to the surface in order to generate the same results as if the feature was turned off. Which gives the good results
+		// since any sub height sampling would in pratice requires less and less texture resolution, where as we are only concerned about the highest resolution the material needs.
+		// This can be seen in the debug view mode, by a checkboard of white and cyan (up to green) values. The white value meaning the highest resolution taken is the good one
+		// (blue meaning the texture has more resolution than required). Checkboard are only possible when a texture is sampled several times, like in parallax.
+		//
+		// Additionnal to affecting the view vector, it also forces a constant world position value, zeroing any textcoord scales that depends on the world position (as the UV don't change).
+		// This is alright thought since the uniform quad can obviously not compute a valid mapping for world space texture mapping (only rendering the mesh at its world position could fix that).
+		// The zero scale will be caught as an error, and the computed scale will fallback to 1.f
+		ResolvedView.SVPositionToTranslatedWorld = float4x4(0,0,0,0,0,0,0,0,0,0,0,0,0,0,1,1);
+
+#if NUM_MATERIAL_TEXCOORDS > 0
+		// Set all texcoord but one to 0. This allows the analysis to figure out which texcoord is driving the lookup.
+		int CurrentCoordIndex = (PixelPosition.y / TILE_RESOLUTION) % MAX_NUM_TEX_COORD;
+
+		UNROLL
+		for(int CoordinateIndex = 0; CoordinateIndex < NUM_MATERIAL_TEXCOORDS; ++CoordinateIndex)
+		{
+			// The quad material used only has a valid TexCoord for UV 0
+			MaterialParameters.TexCoords[CoordinateIndex] = CoordinateIndex != CurrentCoordIndex ? 0 : DebugInputs.TexCoord01.xy;
+		}
+#endif
+	}	
+
 	FPixelMaterialInputs PixelMaterialInputs;
 	CalcMaterialParameters(MaterialParameters, PixelMaterialInputs, DebugInputs.SvPosition, bIsFrontFace);
 
@@ -110,17 +157,27 @@ void Main(
 	float MaterialAO = GetMaterialAmbientOcclusion(PixelMaterialInputs);
 	float PixelDepthOffset = GetMaterialPixelDepthOffset(PixelMaterialInputs);
 
-	// Output accuracy
-	if (GPUScale != INITIAL_GPU_SCALE)
+	if (TextureAnalysisIndex == -1)
 	{
-		float Accuracy = clamp(log2(GPUScale) - log2(CPUScale), -1.99, 1.99);
-		int ColorIndex = floor(Accuracy) + 2;
-
-		Result = TexSample * lerp(AccuracyColors[ColorIndex].rgb, AccuracyColors[ColorIndex + 1].rgb, frac(Accuracy));
+		// Output the scales
+		OutColor = ScalesPerIndex;
 	}
+	else
+	{
+		float PixelScale = (PixelPosition.x & 0x08) == (PixelPosition.y & 0x08) ? MinScale : MaxScale;
+
+		// Output accuracy. If the GPU Scale is INITIAL_GPU_SCALE, then the texcoord was not updated.
+		if (MinScale != INITIAL_GPU_SCALE)
+		{
+			TexSample = TextureAnalysisIndex == MAX_NUM_TEXTURE_REGISTER ? lerp(.4f, 1.f, saturate(Luminance(BaseColor))) : TexSample;
 
-	OutColor = RETURN_COLOR(float4(Result, .25f));
+			float Accuracy = clamp(log2(PixelScale), -1.99, 1.99);
+			int ColorIndex = floor(Accuracy) + 2;
+			Result = TexSample * lerp(AccuracyColors[ColorIndex].rgb, AccuracyColors[ColorIndex + 1].rgb, frac(Accuracy));
+		}
 
+		OutColor = RETURN_COLOR(float4(Result, .25f));
+	}
 }
 
 
diff --git a/Engine/Shaders/TexelFactorAccuracyPixelShader.usf b/Engine/Shaders/TexelFactorAccuracyPixelShader.usf
--- a/Engine/Shaders/TexelFactorAccuracyPixelShader.usf
+++ b/Engine/Shaders/TexelFactorAccuracyPixelShader.usf
@@ -30,8 +30,8 @@ void Main(
 	{
 		float3 WorldPosition = SvPositionToResolvedTranslatedWorld(DebugInputs.SvPosition);
 
-		float2 CoordDDX = ddx(DebugInputs.TexCoord0.xy);
-		float2 CoordDDY = ddy(DebugInputs.TexCoord0.xy);
+		float2 CoordDDX = ddx(DebugInputs.TexCoord01.xy);
+		float2 CoordDDY = ddy(DebugInputs.TexCoord01.xy);
 		float3 WorldPosDDX = ddx(WorldPosition);
 		float3 WorldPosDDY = ddy(WorldPosition);
 
diff --git a/Engine/Shaders/VelocityShader.usf b/Engine/Shaders/VelocityShader.usf
--- a/Engine/Shaders/VelocityShader.usf
+++ b/Engine/Shaders/VelocityShader.usf
@@ -111,7 +111,11 @@ struct FVelocityVSToPS
 void MainVertexShader(
 	FVertexFactoryInput Input,
 	OPTIONAL_VertexID
-	out FVertexOutput Output )
+	out FVertexOutput Output 
+#if USE_GLOBAL_CLIP_PLANE && !USING_TESSELLATION
+	, out float OutGlobalClipPlaneDistance : SV_ClipDistance
+#endif
+	)
 {
 	ResolvedView = ResolveView();
 	
@@ -156,6 +160,10 @@ void MainVertexShader(
 		}
 #endif
 
+#if USE_GLOBAL_CLIP_PLANE && !USING_TESSELLATION
+		OutGlobalClipPlaneDistance = dot(ResolvedView.GlobalClippingPlane, float4(TranslatedWorldPosition.xyz - ResolvedView.PreViewTranslation.xyz, 1));
+#endif
+
 		Output.Interps.PackedVelocityA = ScreenPos;
 		Output.Interps.PackedVelocityC = PrevScreenPosObj;
 	}
diff --git a/Engine/Shaders/WantedMipsAccuracyPixelShader.usf b/Engine/Shaders/WantedMipsAccuracyPixelShader.usf
--- a/Engine/Shaders/WantedMipsAccuracyPixelShader.usf
+++ b/Engine/Shaders/WantedMipsAccuracyPixelShader.usf
@@ -28,8 +28,8 @@ void Main(
 
 	if (CPUWantedMips >= 0)
 	{
-		float2 CoordDDX = ddx(DebugInputs.TexCoord0.xy);
-		float2 CoordDDY = ddy(DebugInputs.TexCoord0.xy);
+		float2 CoordDDX = ddx(DebugInputs.TexCoord01.xy);
+		float2 CoordDDY = ddy(DebugInputs.TexCoord01.xy);
 
 		float TexelRatio = max(1 / length(CoordDDX), 1 / length(CoordDDY));
 
diff --git a/Engine/Source/Developer/MaterialUtilities/Private/MaterialUtilities.cpp b/Engine/Source/Developer/MaterialUtilities/Private/MaterialUtilities.cpp
--- a/Engine/Source/Developer/MaterialUtilities/Private/MaterialUtilities.cpp
+++ b/Engine/Source/Developer/MaterialUtilities/Private/MaterialUtilities.cpp
@@ -1696,6 +1696,247 @@ void FMaterialUtilities::ResizeFlattenMaterial(FFlattenMaterial& InFlattenMateri
 	}
 }
 
+/** Computes the uniform scale from the input scales,  if one exists. */
+static bool GetUniformScale(const TArray<float> Scales, float& UniformScale)
+{
+	if (Scales.Num())
+	{
+		float Average = 0;
+		float Mean = 0;
+
+		for (float V : Scales)
+		{
+			Average += V;
+		}
+		Average /= (float)Scales.Num();
+
+		for (float V : Scales)
+		{
+			Mean += FMath::Abs(V - Average);
+		}
+		Mean /= (float)Scales.Num();
+
+		if (Mean * 15.f < Average) // If they are almost all the same
+		{
+			UniformScale = Average;
+			return true;
+		}
+		else // Otherwise do a much more expensive test by counting the number of similar values
+		{
+			// Try to find a small range where 80% of values fit within.
+			const int32 TryThreshold = FMath::CeilToInt(.80f * (float)Scales.Num());
+
+			int32 NextTryDomain = Scales.Num();
+
+			float NextTryMinV = 1024;
+			for (float V : Scales)
+			{
+				NextTryMinV = FMath::Min<float>(V, NextTryMinV);
+			}
+
+			while (NextTryDomain >= TryThreshold) // Stop the search it is garantied to fail.
+			{
+				float TryMinV = NextTryMinV;
+				float TryMaxV = TryMinV * 1.25f;
+				int32 TryMatches = 0;
+				NextTryMinV = 1024;
+				NextTryDomain = 0;
+				for (float V : Scales)
+				{
+					if (TryMinV <= V && V <= TryMaxV)
+					{
+						++TryMatches;
+					}
+
+					if (V > TryMinV)
+					{
+						NextTryMinV = FMath::Min<float>(V, NextTryMinV);
+						++NextTryDomain;
+					}
+				}
+
+				if (TryMatches >= TryThreshold)
+				{
+					UniformScale = TryMinV;
+					return true;
+				}
+			}
+		}
+	}
+	return false;
+}
+
+bool FMaterialUtilities::ExportMaterialTexCoordScales(UMaterialInterface* InMaterial, EMaterialQualityLevel::Type QualityLevel, ERHIFeatureLevel::Type FeatureLevel, TArray<FMaterialTexCoordBuildInfo>& OutScales)
+{
+	TArray<FFloat16Color> RenderedVectors;
+
+	TArray<UTexture*> Textures;
+	TArray< TArray<int32> > Indices;
+	InMaterial->GetUsedTexturesAndIndices(Textures, Indices, QualityLevel, FeatureLevel);
+
+	check(Textures.Num() >= Indices.Num()); // Can't have indices if no texture.
+
+	// Clear any entry not related to UTexture2D
+	for (int32 Index = 0; Index < Textures.Num(); ++Index)
+	{
+		if (!Cast<UTexture2D>(Textures[Index]))
+		{
+			Indices[Index].Empty(); 
+		}
+	}
+
+	const int32 SCALE_PRECISION = 64.f;
+	const int32 TILE_RESOLUTION = FMaterialTexCoordBuildInfo::TILE_RESOLUTION;
+	const int32 INITIAL_GPU_SCALE = FMaterialTexCoordBuildInfo::INITIAL_GPU_SCALE;
+	const int32 MAX_NUM_TEX_COORD = FMaterialTexCoordBuildInfo::MAX_NUM_TEX_COORD;
+	const int32 MAX_NUM_TEXTURE_REGISTER = FMaterialTexCoordBuildInfo::MAX_NUM_TEXTURE_REGISTER;
+
+	const bool bUseMetrics = CVarStreamingUseNewMetrics.GetValueOnGameThread() != 0;
+
+	int32 MaxRegisterIndex = INDEX_NONE;
+
+	for (const TArray<int32>& TextureIndices : Indices)
+	{
+		for (int32 RegisterIndex : TextureIndices)
+		{
+			MaxRegisterIndex = FMath::Max<int32>(RegisterIndex, MaxRegisterIndex);
+		}
+	}
+
+	if (MaxRegisterIndex == INDEX_NONE)
+	{
+		return false;
+	}
+
+	TBitArray<> RegisterInUse(false, MaxRegisterIndex + 1);
+
+	// Set the validity flag.
+	for (const TArray<int32>& TextureIndices : Indices)
+	{
+		for (int32 RegisterIndex : TextureIndices)
+		{
+			RegisterInUse[RegisterIndex] = true;
+		}
+	}
+
+
+	const int32 NumTileX = (MaxRegisterIndex / 4 + 1);
+	const int32 NumTileY = MAX_NUM_TEX_COORD;
+	FIntPoint RenderTargetSize(TILE_RESOLUTION * NumTileX, TILE_RESOLUTION * NumTileY);
+
+	// Render the vectors
+	{
+		// The rendertarget contain factors stored in XYZW. Every X tile maps to another group : (0, 1, 2, 3), (4, 5, 6, 7), ...
+		UTextureRenderTarget2D* RenderTarget = CreateRenderTarget(true, false, PF_FloatRGBA, RenderTargetSize);
+
+		// Allocate the render output.
+		RenderedVectors.Empty(RenderTargetSize.X * RenderTargetSize.Y);
+
+		FMaterialRenderProxy* MaterialProxy = InMaterial->GetRenderProxy(false, false);
+		if (!MaterialProxy)
+		{
+			return false;
+		}
+
+		FBox2D DummyBounds(FVector2D(0, 0), FVector2D(1, 1));
+		TArray<FVector2D> EmptyTexCoords;
+		FMaterialMergeData MaterialData(InMaterial, nullptr, nullptr, 0, DummyBounds, EmptyTexCoords);
+
+		CurrentlyRendering = true;
+		bool bResult = FMeshRenderer::RenderMaterialTexCoordScales(MaterialData, MaterialProxy, RenderTarget, RenderedVectors);
+		CurrentlyRendering = false;
+
+		if (!bResult)
+		{
+			return false;
+		}
+	}
+
+	OutScales.AddDefaulted(MaxRegisterIndex + 1);
+
+	TArray<float> TextureIndexScales;
+
+	// Now compute the scale for each
+	for (int32 RegisterIndex = 0; RegisterIndex <= MaxRegisterIndex; ++RegisterIndex)
+	{
+		if (!RegisterInUse[RegisterIndex]) continue;
+
+		// If nothing works, this will fallback to 1.f
+		OutScales[RegisterIndex].Scale = 1.f;
+		OutScales[RegisterIndex].Index = 0;
+
+		if (!bUseMetrics) continue; // Old metrics would always use 1.f
+
+		int32 TextureTile = RegisterIndex / 4;
+		int32 ComponentIndex = RegisterIndex % 4;
+
+		bool bSuccess = false;
+		for (int32 CoordIndex = 0; CoordIndex < MAX_NUM_TEX_COORD && !bSuccess; ++CoordIndex)
+		{
+			TextureIndexScales.Empty(TILE_RESOLUTION * TILE_RESOLUTION);
+			for (int32 TexelX = 0; TexelX < TILE_RESOLUTION; ++TexelX)
+			{
+				for (int32 TexelY = 0; TexelY < TILE_RESOLUTION; ++TexelY)
+				{
+					int32 TexelIndex = TextureTile * TILE_RESOLUTION + TexelX + (TexelY + CoordIndex * TILE_RESOLUTION) * RenderTargetSize.X;
+					FFloat16Color& Scale16 = RenderedVectors[TexelIndex];
+
+					float TexelScale = 0;
+					if (ComponentIndex == 0) TexelScale = Scale16.R.GetFloat();
+					if (ComponentIndex == 1) TexelScale = Scale16.G.GetFloat();
+					if (ComponentIndex == 2) TexelScale = Scale16.B.GetFloat();
+					if (ComponentIndex == 3) TexelScale = Scale16.A.GetFloat();
+
+					// Quantize scale to converge faster in the TryLogic
+					TexelScale = FMath::RoundToFloat(TexelScale * SCALE_PRECISION) / SCALE_PRECISION;
+
+					if (TexelScale > 0 && TexelScale < INITIAL_GPU_SCALE)
+					{
+						TextureIndexScales.Push(TexelScale);
+					}
+				}
+			}
+
+			if (GetUniformScale(TextureIndexScales, OutScales[RegisterIndex].Scale))
+			{
+				OutScales[RegisterIndex].Index = CoordIndex;
+				bSuccess = true;
+			}
+		}
+
+		// If we couldn't find the scale, then output a warning detailing which index, texture, material is having an issue.
+		if (!bSuccess)
+		{
+			UTexture* FailedTexture = nullptr;
+			for (int32 Index = 0; Index < Indices.Num() && !FailedTexture; ++Index)
+			{
+				for (int32 SubIndex : Indices[Index])
+				{
+					if (SubIndex == RegisterIndex)
+					{
+						FailedTexture = Textures[Index];
+						break;
+					}
+				}
+			}
+
+			if (FailedTexture)
+			{
+				if (TextureIndexScales.Num())
+				{
+					UE_LOG(LogMaterialUtilities, Warning, TEXT("ExportMaterialTexCoordScales: Failed to find constant scale for texture index %d of material %s (bound to texture %s)"), RegisterIndex, *InMaterial->GetName(), *FailedTexture->GetName());
+				}
+				else
+				{
+					UE_LOG(LogMaterialUtilities, Warning, TEXT("ExportMaterialTexCoordScales: Failed to generate scales for texture index %d of material %s (bound to texture %s)"), RegisterIndex, *InMaterial->GetName(), *FailedTexture->GetName());
+				}
+			}
+		}
+	}
+
+	return true;
+}
+
 bool FMaterialUtilities::ExportMaterial(struct FMaterialMergeData& InMaterialData, FFlattenMaterial& OutFlattenMaterial, struct FExportMaterialProxyCache* ProxyCache)
 {
 	UMaterialInterface* Material = InMaterialData.Material;
diff --git a/Engine/Source/Developer/MaterialUtilities/Private/MeshRendering.cpp b/Engine/Source/Developer/MaterialUtilities/Private/MeshRendering.cpp
--- a/Engine/Source/Developer/MaterialUtilities/Private/MeshRendering.cpp
+++ b/Engine/Source/Developer/MaterialUtilities/Private/MeshRendering.cpp
@@ -657,7 +657,7 @@ public:
 		// Check if material is TwoSided - single-sided materials should be rendered with normal and reverse
 		// triangle corner orders, to avoid problems with inside-out meshes or mesh parts. Note:
 		// FExportMaterialProxy::GetMaterial() (which is really called here) ignores 'InFeatureLevel' parameter.
-		const FMaterial* Material = Data.MaterialRenderProxy->GetMaterial(ERHIFeatureLevel::SM5);
+		const FMaterial* Material = Data.MaterialRenderProxy->GetMaterial(GMaxRHIFeatureLevel);
 		bool bIsMaterialTwoSided = Material->IsTwoSided();
 
 		TArray<FMaterialMeshVertex> Verts;
@@ -790,6 +790,7 @@ public:
 
 bool FMeshRenderer::RenderMaterial(struct FMaterialMergeData& InMaterialData, FMaterialRenderProxy* InMaterialProxy, EMaterialProperty InMaterialProperty, UTextureRenderTarget2D* InRenderTarget, TArray<FColor>& OutBMP)
 {
+	check(IsInGameThread());
 	check(InRenderTarget);
 	FTextureRenderTargetResource* RTResource = InRenderTarget->GameThread_GetRenderTargetResource();
 
@@ -939,3 +940,50 @@ bool FMeshRenderer::RenderMaterial(struct FMaterialMergeData& InMaterialData, FM
 #endif // SAVE_INTERMEDIATE_TEXTURES
 	return result;
 }
+
+bool FMeshRenderer::RenderMaterialTexCoordScales(struct FMaterialMergeData& InMaterialData, FMaterialRenderProxy* InMaterialProxy, UTextureRenderTarget2D* InRenderTarget, TArray<FFloat16Color>& OutScales)
+{
+	check(IsInGameThread());
+	check(InRenderTarget);
+	// create ViewFamily
+	float CurrentRealTime = 0.f;
+	float CurrentWorldTime = 0.f;
+	float DeltaWorldTime = 0.f;
+
+	// Create a canvas for the render target and clear it to black
+	FTextureRenderTargetResource* RTResource = InRenderTarget->GameThread_GetRenderTargetResource();
+	FCanvas Canvas(RTResource, NULL, FApp::GetCurrentTime() - GStartTime, FApp::GetDeltaTime(), FApp::GetCurrentTime() - GStartTime, GMaxRHIFeatureLevel);
+	const FRenderTarget* CanvasRenderTarget = Canvas.GetRenderTarget();
+	Canvas.Clear(FLinearColor::Black);
+
+	// Set show flag view mode to output tex coord scale
+	FEngineShowFlags ShowFlags(ESFIM_Game);
+	ApplyViewMode(VMI_TexCoordScaleAccuracy, false, ShowFlags);
+	ShowFlags.TexCoordAnalysis = true; // This will bind the DVSM_TexCoordScaleAnalysis
+
+	FSceneViewFamily ViewFamily(FSceneViewFamily::ConstructionValues(CanvasRenderTarget, nullptr, ShowFlags)
+		.SetWorldTimes(CurrentWorldTime, DeltaWorldTime, CurrentRealTime)
+		.SetGammaCorrection(CanvasRenderTarget->GetDisplayGamma()));
+		
+	// add item for rendering
+	FMeshMaterialRenderItem::EnqueueMaterialRender(
+		&Canvas,
+		&ViewFamily,
+		InMaterialData.Mesh,
+		InMaterialData.LODModel,
+		InMaterialData.MaterialIndex,
+		InMaterialData.TexcoordBounds,
+		InMaterialData.TexCoords,
+		FVector2D(InRenderTarget->SizeX, InRenderTarget->SizeY),
+		InMaterialProxy
+		);
+
+	// rendering is performed here
+	Canvas.Flush_GameThread();
+
+	FlushRenderingCommands();
+	Canvas.SetRenderTarget_GameThread(NULL);
+	FlushRenderingCommands();
+
+	return RTResource->ReadFloat16Pixels(OutScales);
+}
\ No newline at end of file
diff --git a/Engine/Source/Developer/MaterialUtilities/Public/MaterialUtilities.h b/Engine/Source/Developer/MaterialUtilities/Public/MaterialUtilities.h
--- a/Engine/Source/Developer/MaterialUtilities/Public/MaterialUtilities.h
+++ b/Engine/Source/Developer/MaterialUtilities/Public/MaterialUtilities.h
@@ -8,6 +8,7 @@
 #include "Runtime/Engine/Classes/Engine/Texture.h"
 #include "Runtime/Engine/Public/SceneTypes.h"
 #include "Runtime/Engine/Classes/Materials/Material.h"
+#include "Runtime/Engine/Classes/Engine/TextureStreamingTypes.h"
 #include "RawMesh.h"
 
 /** Structure used for storing intermediate baked down material data/samples*/
@@ -325,6 +326,18 @@ public:
 	* @param InFlattenMaterial				Flatten material to optimize
 	*/
 	static void ResizeFlattenMaterial(FFlattenMaterial& InFlattenMaterial, const struct FMeshProxySettings& InMeshProxySettings);
+
+	/**
+	* Get the material texcoord scales applied on each textures
+	*
+	* @param InMaterial			Target material
+	* @param QualityLevel		Quality level used for the shader profiling.
+	* @param FeatureLevel		Feature level used for the shader profiling.
+	* @param OutScales			TheOutput array of rendered samples	
+	* @return					Whether operation was successful
+	*/
+	static bool ExportMaterialTexCoordScales(UMaterialInterface* InMaterial, EMaterialQualityLevel::Type QualityLevel, ERHIFeatureLevel::Type FeatureLevel, TArray<FMaterialTexCoordBuildInfo>& OutScales);
+
 private:
 	
 	/**
diff --git a/Engine/Source/Developer/MaterialUtilities/Public/MeshRendering.h b/Engine/Source/Developer/MaterialUtilities/Public/MeshRendering.h
--- a/Engine/Source/Developer/MaterialUtilities/Public/MeshRendering.h
+++ b/Engine/Source/Developer/MaterialUtilities/Public/MeshRendering.h
@@ -16,8 +16,22 @@ class FMeshRenderer
 {
 public:
 	/** Renders out textures for each material property for the given material, using the given mesh data or by using a simple tile rendering approach */
-	MATERIALUTILITIES_API static bool RenderMaterial(FMaterialMergeData& InMaterialData, FMaterialRenderProxy* InMaterialProxy,
-		EMaterialProperty InMaterialProperty, UTextureRenderTarget2D* InRenderTarget, TArray<FColor>& OutBMP);	
+	MATERIALUTILITIES_API static bool RenderMaterial(
+		FMaterialMergeData& InMaterialData, 
+		FMaterialRenderProxy* InMaterialProxy,
+		EMaterialProperty InMaterialProperty, 
+		UTextureRenderTarget2D* InRenderTarget, 
+		TArray<FColor>& OutBMP
+		);	
+
+	/** Renders out texcoord scales */
+	MATERIALUTILITIES_API static bool RenderMaterialTexCoordScales(
+		FMaterialMergeData& InMaterialData, 
+		FMaterialRenderProxy* InMaterialProxy,
+		UTextureRenderTarget2D* InRenderTarget, 
+		TArray<FFloat16Color>& OutScales
+		);	
+
 private:
 	/** This class never needs to be instantiated. */
 	FMeshRenderer() {}
diff --git a/Engine/Source/Developer/ShaderCompilerCommon/Private/HlslUtils.cpp b/Engine/Source/Developer/ShaderCompilerCommon/Private/HlslUtils.cpp
--- a/Engine/Source/Developer/ShaderCompilerCommon/Private/HlslUtils.cpp
+++ b/Engine/Source/Developer/ShaderCompilerCommon/Private/HlslUtils.cpp
@@ -80,6 +80,37 @@ namespace CrossCompiler
 	}
 }
 
+// Returns "TEXCOORD" if Semantic is "TEXCOORD4" and OutStartIndex=4; return nullptr if the semantic didn't have a number
+static inline TCHAR* GetNonDigitSemanticPrefix(CrossCompiler::FLinearAllocator* Allocator, const TCHAR* Semantic, uint32& OutStartIndex)
+{
+	const TCHAR* StartOfDigit = Semantic;
+	do 
+	{
+		if (*StartOfDigit >= '0' && *StartOfDigit <= '9')
+		{
+			break;
+		}
+		++StartOfDigit;
+	}
+	while (*StartOfDigit);
+
+	if (!*StartOfDigit)
+	{
+		return nullptr;
+	}
+
+	OutStartIndex = FCString::Atoi(StartOfDigit);
+	TCHAR* Prefix = Allocator->Strdup(Semantic);
+	Prefix[StartOfDigit - Semantic] = 0;
+	return Prefix;
+}
+
+
+static inline TCHAR* MakeIndexedSemantic(CrossCompiler::FLinearAllocator* Allocator, const TCHAR* Semantic, uint32 Index)
+{
+	FString Out = FString::Printf(TEXT("%s%d"), Semantic, Index);
+	return Allocator->Strdup(Out);
+}
 
 struct FRemoveAlgorithm
 {
@@ -238,7 +269,20 @@ struct FRemoveAlgorithm
 		return false;
 	};
 
-	bool CopyMember(CrossCompiler::AST::FDeclaration* Declaration, const TCHAR* DestPrefix, const TCHAR* SourcePrefix, FBodyContext& BodyContext)
+	static bool IsSubstringInArray(const TArray<FString>& Array, const TCHAR* Semantic)
+	{
+		for (const FString& String : Array)
+		{
+			if (FCString::Stristr(*String, Semantic) != NULL)
+			{
+				return true;
+			}
+		}
+
+		return false;
+	};
+
+	bool CopyMember(CrossCompiler::AST::FDeclaration* Declaration, const TCHAR* DestPrefix, const TCHAR* SourcePrefix, TArray<CrossCompiler::AST::FNode*>& InstructionList)
 	{
 		using namespace CrossCompiler::AST;
 
@@ -265,7 +309,7 @@ struct FRemoveAlgorithm
 				auto* LHS = MakeIdentifierExpression(Allocator, *LHSElement, SourceInfo);
 				auto* RHS = MakeIdentifierExpression(Allocator, *RHSElement, SourceInfo);
 				auto* Assignment = new(Allocator) FBinaryExpression(Allocator, EOperators::Assign, LHS, RHS, SourceInfo);
-				BodyContext.PostInstructions.Add(new(Allocator) FExpressionStatement(Allocator, Assignment, SourceInfo));
+				InstructionList.Add(new(Allocator) FExpressionStatement(Allocator, Assignment, SourceInfo));
 			}
 		}
 		else
@@ -273,7 +317,7 @@ struct FRemoveAlgorithm
 			auto* LHS = MakeIdentifierExpression(Allocator, *LHSName, SourceInfo);
 			auto* RHS = MakeIdentifierExpression(Allocator, *RHSName, SourceInfo);
 			auto* Assignment = new(Allocator) FBinaryExpression(Allocator, EOperators::Assign, LHS, RHS, SourceInfo);
-			BodyContext.PostInstructions.Add(new(Allocator) FExpressionStatement(Allocator, Assignment, SourceInfo));
+			InstructionList.Add(new(Allocator) FExpressionStatement(Allocator, Assignment, SourceInfo));
 		}
 
 		return true;
@@ -376,6 +420,7 @@ struct FRemoveAlgorithm
 struct FRemoveUnusedOutputs : FRemoveAlgorithm
 {
 	const TArray<FString>& UsedOutputs;
+	const TArray<FString>& Exceptions;
 
 	struct FOutputsBodyContext : FBodyContext
 	{
@@ -399,8 +444,9 @@ struct FRemoveUnusedOutputs : FRemoveAlgorithm
 		}
 	};
 
-	FRemoveUnusedOutputs(const TArray<FString>& InUsedOutputs) :
-		UsedOutputs(InUsedOutputs)
+	FRemoveUnusedOutputs(const TArray<FString>& InUsedOutputs, const TArray<FString>& InExceptions) :
+		UsedOutputs(InUsedOutputs),
+		Exceptions(InExceptions)
 	{
 	}
 
@@ -706,7 +752,7 @@ struct FRemoveUnusedOutputs : FRemoveAlgorithm
 
 	bool IsSemanticUsed(const TCHAR* SemanticName)
 	{
-		if (bLeaveAllUsed || IsStringInArray(UsedOutputs, SemanticName))
+		if (bLeaveAllUsed || IsStringInArray(UsedOutputs, SemanticName) || IsSubstringInArray(Exceptions, SemanticName))
 		{
 			return true;
 		}
@@ -736,7 +782,60 @@ struct FRemoveUnusedOutputs : FRemoveAlgorithm
 				check(MemberDeclaration);
 				if (MemberDeclaration->Semantic)
 				{
-					if (IsSemanticUsed(MemberDeclaration->Semantic))
+					if (MemberDeclaration->bIsArray)
+					{
+						uint32 ArrayLength = 0;
+						if (!GetArrayLength(MemberDeclaration, ArrayLength))
+						{
+							return false;
+						}
+
+						uint32 StartIndex = 0;
+						TCHAR* ElementSemanticPrefix = GetNonDigitSemanticPrefix(Allocator, MemberDeclaration->Semantic, StartIndex);
+						if (!ElementSemanticPrefix)
+						{
+							Errors.Add(FString::Printf(TEXT("RemoveUnusedOutputs: Member (%s) %s : %s is expected to have an indexed semantic!"), MemberDeclarator->Type->Specifier->TypeName, MemberDeclaration->Identifier, MemberDeclaration->Semantic));
+
+							// Fatal: Array of non-indexed semantic (eg float4 Colors[4] : MYSEMANTIC; )
+							// Assume semantic is used and just fallback
+							auto* NewDeclaratorList = new(Allocator) FDeclaratorList(Allocator, MemberDeclarator->SourceInfo);
+							NewDeclaratorList->Type = CloneType(MemberDeclarator->Type);
+							NewDeclaratorList->Declarations.Add(MemberDeclaration);
+							DestStruct->Declarations.Add(NewDeclaratorList);
+
+							CopyMember(MemberDeclaration, SourcePrefix, DestPrefix, BodyContext.PostInstructions);
+						}
+
+						for (uint32 Index = 0; Index < ArrayLength; ++Index)
+						{
+							TCHAR* ElementSemantic = MakeIndexedSemantic(Allocator, ElementSemanticPrefix, StartIndex + Index);
+							if (IsSemanticUsed(ElementSemantic))
+							{
+								auto* NewMemberDeclaration = new(Allocator) FDeclaration(Allocator, MemberDeclaration->SourceInfo);
+								NewMemberDeclaration->Semantic = ElementSemantic;
+								NewMemberDeclaration->Identifier = Allocator->Strdup(FString::Printf(TEXT("%s_%d"), MemberDeclaration->Identifier, Index));
+
+								// Add member to struct
+								auto* NewDeclaratorList = new(Allocator) FDeclaratorList(Allocator, MemberDeclarator->SourceInfo);
+								NewDeclaratorList->Type = CloneType(MemberDeclarator->Type);
+								NewDeclaratorList->Declarations.Add(NewMemberDeclaration);
+								DestStruct->Declarations.Add(NewDeclaratorList);
+
+								FString LHSElement = FString::Printf(TEXT("%s.%s"), DestPrefix, NewMemberDeclaration->Identifier);
+								FString RHSElement = FString::Printf(TEXT("%s.%s[%d]"), SourcePrefix, MemberDeclaration->Identifier, Index);
+
+								auto* LHS = MakeIdentifierExpression(Allocator, *LHSElement, SourceInfo);
+								auto* RHS = MakeIdentifierExpression(Allocator, *RHSElement, SourceInfo);
+								auto* Assignment = new(Allocator) FBinaryExpression(Allocator, EOperators::Assign, LHS, RHS, SourceInfo);
+								BodyContext.PostInstructions.Add(new(Allocator) FExpressionStatement(Allocator, Assignment, SourceInfo));
+							}
+							else
+							{
+								RemovedSemantics.Add(ElementSemantic);
+							}
+						}
+					}
+					else if (IsSemanticUsed(MemberDeclaration->Semantic))
 					{
 						// Add member to struct
 						auto* NewDeclaratorList = new(Allocator) FDeclaratorList(Allocator, MemberDeclarator->SourceInfo);
@@ -744,7 +843,7 @@ struct FRemoveUnusedOutputs : FRemoveAlgorithm
 						NewDeclaratorList->Declarations.Add(MemberDeclaration);
 						DestStruct->Declarations.Add(NewDeclaratorList);
 
-						CopyMember(MemberDeclaration, DestPrefix, SourcePrefix, BodyContext);
+						CopyMember(MemberDeclaration, DestPrefix, SourcePrefix, BodyContext.PostInstructions);
 					}
 					else
 					{
@@ -793,10 +892,10 @@ static void HlslParserCallbackWrapperRemoveOutputs(void* CallbackData, CrossComp
 }
 
 
-bool RemoveUnusedOutputs(FString& InOutSourceCode, const TArray<FString>& InUsedOutputs, FString& EntryPoint, TArray<FString>& OutErrors)
+bool RemoveUnusedOutputs(FString& InOutSourceCode, const TArray<FString>& InUsedOutputs, const TArray<FString>& InExceptions, FString& EntryPoint, TArray<FString>& OutErrors)
 {
 	FString DummyFilename(TEXT("RemoveUnusedOutputs.usf"));
-	FRemoveUnusedOutputs Data(InUsedOutputs);
+	FRemoveUnusedOutputs Data(InUsedOutputs, InExceptions);
 	Data.EntryPoint = EntryPoint;
 	CrossCompiler::FCompilerMessages Messages;
 	if (!CrossCompiler::Parser::Parse(InOutSourceCode, DummyFilename, Messages, HlslParserCallbackWrapperRemoveOutputs, &Data))
@@ -1134,7 +1233,60 @@ struct FRemoveUnusedInputs : FRemoveAlgorithm
 				check(MemberDeclaration);
 				if (MemberDeclaration->Semantic)
 				{
-					if (IsSemanticUsed(MemberDeclaration->Semantic))
+					if (MemberDeclaration->bIsArray)
+					{
+						uint32 ArrayLength = 0;
+						if (!GetArrayLength(MemberDeclaration, ArrayLength))
+						{
+							return false;
+						}
+
+						uint32 StartIndex = 0;
+						TCHAR* ElementSemanticPrefix = GetNonDigitSemanticPrefix(Allocator, MemberDeclaration->Semantic, StartIndex);
+						if (!ElementSemanticPrefix)
+						{
+							Errors.Add(FString::Printf(TEXT("RemoveUnusedInputs: Member (%s) %s : %s is expected to have an indexed semantic!"), MemberDeclarator->Type->Specifier->TypeName, MemberDeclaration->Identifier, MemberDeclaration->Semantic));
+
+							// Fatal: Array of non-indexed semantic (eg float4 Colors[4] : MYSEMANTIC; )
+							// Assume semantic is used and just fallback
+							auto* NewDeclaratorList = new(Allocator) FDeclaratorList(Allocator, MemberDeclarator->SourceInfo);
+							NewDeclaratorList->Type = CloneType(MemberDeclarator->Type);
+							NewDeclaratorList->Declarations.Add(MemberDeclaration);
+							DestStruct->Declarations.Add(NewDeclaratorList);
+
+							CopyMember(MemberDeclaration, SourcePrefix, DestPrefix, BodyContext.PreInstructions);
+						}
+
+						for (uint32 Index = 0; Index < ArrayLength; ++Index)
+						{
+							TCHAR* ElementSemantic = MakeIndexedSemantic(Allocator, ElementSemanticPrefix, StartIndex + Index);
+							if (IsSemanticUsed(ElementSemantic))
+							{
+								auto* NewMemberDeclaration = new(Allocator) FDeclaration(Allocator, MemberDeclaration->SourceInfo);
+								NewMemberDeclaration->Semantic = ElementSemantic;
+								NewMemberDeclaration->Identifier = Allocator->Strdup(FString::Printf(TEXT("%s_%d"), MemberDeclaration->Identifier, Index));
+
+								// Add member to struct
+								auto* NewDeclaratorList = new(Allocator) FDeclaratorList(Allocator, MemberDeclarator->SourceInfo);
+								NewDeclaratorList->Type = CloneType(MemberDeclarator->Type);
+								NewDeclaratorList->Declarations.Add(NewMemberDeclaration);
+								DestStruct->Declarations.Add(NewDeclaratorList);
+
+								FString LHSElement = FString::Printf(TEXT("%s.%s[%d]"), SourcePrefix, MemberDeclaration->Identifier, Index);
+								FString RHSElement = FString::Printf(TEXT("%s.%s"), DestPrefix, NewMemberDeclaration->Identifier);
+
+								auto* LHS = MakeIdentifierExpression(Allocator, *LHSElement, SourceInfo);
+								auto* RHS = MakeIdentifierExpression(Allocator, *RHSElement, SourceInfo);
+								auto* Assignment = new(Allocator) FBinaryExpression(Allocator, EOperators::Assign, LHS, RHS, SourceInfo);
+								BodyContext.PreInstructions.Add(new(Allocator) FExpressionStatement(Allocator, Assignment, SourceInfo));
+							}
+							else
+					{
+								RemovedSemantics.Add(ElementSemantic);
+							}
+						}
+					}
+					else if (IsSemanticUsed(MemberDeclaration->Semantic))
 					{
 						// Add member to struct
 						auto* NewDeclaratorList = new(Allocator) FDeclaratorList(Allocator, MemberDeclarator->SourceInfo);
@@ -1142,7 +1294,7 @@ struct FRemoveUnusedInputs : FRemoveAlgorithm
 						NewDeclaratorList->Declarations.Add(MemberDeclaration);
 						DestStruct->Declarations.Add(NewDeclaratorList);
 
-						CopyMember(MemberDeclaration, SourcePrefix, DestPrefix, BodyContext);
+						CopyMember(MemberDeclaration, SourcePrefix, DestPrefix, BodyContext.PreInstructions);
 					}
 					else
 					{
diff --git a/Engine/Source/Developer/ShaderCompilerCommon/Public/ShaderCompilerCommon.h b/Engine/Source/Developer/ShaderCompilerCommon/Public/ShaderCompilerCommon.h
--- a/Engine/Source/Developer/ShaderCompilerCommon/Public/ShaderCompilerCommon.h
+++ b/Engine/Source/Developer/ShaderCompilerCommon/Public/ShaderCompilerCommon.h
@@ -63,7 +63,7 @@ extern SHADERCOMPILERCOMMON_API int16 GetNumUniformBuffersUsed(const FShaderComp
 //		}
 //
 extern SHADERCOMPILERCOMMON_API bool RemoveUniformBuffersFromSource(FString& SourceCode);
-extern SHADERCOMPILERCOMMON_API bool RemoveUnusedOutputs(FString& InOutSourceCode, const TArray<FString>& InUsedOutputs, FString& InOutEntryPoint, TArray<FString>& OutErrors);
+extern SHADERCOMPILERCOMMON_API bool RemoveUnusedOutputs(FString& InOutSourceCode, const TArray<FString>& InUsedOutputs, const TArray<FString>& InExceptions, FString& InOutEntryPoint, TArray<FString>& OutErrors);
 
 extern SHADERCOMPILERCOMMON_API bool RemoveUnusedInputs(FString& InOutSourceCode, const TArray<FString>& InUsedInputs, FString& InOutEntryPoint, TArray<FString>& OutErrors);
 
diff --git a/Engine/Source/Developer/VulkanShaderFormat/Private/VulkanBackend.cpp b/Engine/Source/Developer/VulkanShaderFormat/Private/VulkanBackend.cpp
--- a/Engine/Source/Developer/VulkanShaderFormat/Private/VulkanBackend.cpp
+++ b/Engine/Source/Developer/VulkanShaderFormat/Private/VulkanBackend.cpp
@@ -964,29 +964,15 @@ class vulkan_ir_gen_glsl_visitor : public ir_visitor
 					check(inner_type->length == 1);
 					const glsl_struct_field* field = &inner_type->fields.structure[0];
 
-					#if VULKAN_HLSLCC_GEN_GLSL_IN_OUT_WITHOUT_STRUCTURES
-						ralloc_asprintf_append(
-							buffer,
-							"layout(location=%d) %s",
-							var->location,				// location number
-							mode_str[var->mode]			// in / out
-							);
+					ralloc_asprintf_append(
+						buffer,
+						"layout(location=%d) %s",
+						var->location,				// location number
+						mode_str[var->mode]			// in / out
+						);
 
-						// Append type to the buffer string
-						print_type_pre(field->type); // float, vec2, vec3 and etc..
-					#else
-						ralloc_asprintf_append(
-							buffer,
-							"layout(location=%d) %s struct { ",
-							var->location,
-							mode_str[var->mode]);
-						ralloc_asprintf_append(buffer, " ");
-						check(strcmp(field->name, "Data") == 0);
-						print_type_pre(field->type);
-						ralloc_asprintf_append(buffer, " Data");
-						print_type_post(field->type);
-						ralloc_asprintf_append(buffer, "; }");
-					#endif
+					// Append type to the buffer string
+					print_type_pre(field->type); // float, vec2, vec3 and etc..
 				}
 				else
 				{
@@ -3844,77 +3830,39 @@ static ir_rvalue* GenShaderInputSemantic(
 	}
 	else if (SemanticArrayIndex == 0)
 	{
-		#if VULKAN_HLSLCC_GEN_GLSL_IN_OUT_WITHOUT_STRUCTURES
-			// This code-section replaces "layout(location=0) in struct { vec4 Data; } in_ATTRIBUTE0;" pattern to
-			// "layout(location=0) in vec4 in_ATTRIBUTE0;".
-
-			if (/*Frequency == HSF_GeometryShader && */SemanticArraySize != 0)
-			{
-				Type = glsl_type::get_array_instance(Type, SemanticArraySize);
-			}
-
-			// Regular attribute
-			Variable = new(ParseState)ir_variable(
-				Type,
-				ralloc_asprintf(ParseState, "in_%s", Semantic),
-				ir_var_in
-				);
-			Variable->read_only = true;
-			Variable->centroid = InputQualifier.Fields.bCentroid;
-			Variable->interpolation = InputQualifier.Fields.InterpolationMode;
-			Variable->is_patch_constant = InputQualifier.Fields.bIsPatchConstant;
-
-			if (ParseState->bGenerateLayoutLocations && !InputQualifier.Fields.bIsPatchConstant)
-			{
-				ConfigureInOutVariableLayout(Frequency, ParseState, Semantic, Variable, ir_var_in);
-			}
-
-			DeclInstructions->push_tail(Variable);
-			ParseState->symbols->add_variable(Variable);
+		// This code-section replaces "layout(location=0) in struct { vec4 Data; } in_ATTRIBUTE0;" pattern to
+		// "layout(location=0) in vec4 in_ATTRIBUTE0;".
 
-			ir_dereference* VariableDeref = new(ParseState)ir_dereference_variable(Variable);
-			if (SemanticArraySize > 0)
-			{
-				// Deref inside array first
-				VariableDeref = new(ParseState) ir_dereference_array(VariableDeref, new(ParseState) ir_constant((unsigned)SemanticArrayIndex));
-			}
-
-		#else
-			glsl_struct_field *StructField = ralloc_array(ParseState, glsl_struct_field, 1);
-			memset(StructField, 0, sizeof(glsl_struct_field));
-			StructField[0].type = Type;
-			StructField[0].name = ralloc_strdup(ParseState, "Data");
+		if (/*Frequency == HSF_GeometryShader && */SemanticArraySize != 0)
+		{
+			Type = glsl_type::get_array_instance(Type, SemanticArraySize);
+		}
 
-			const glsl_type* VariableType = glsl_type::get_record_instance(StructField, 1, ralloc_strdup(ParseState, Semantic));
-			if (SemanticArraySize)
-			{
-				// Pack it into an array too
-				VariableType = glsl_type::get_array_instance(VariableType, SemanticArraySize);
-			}
+		// Regular attribute
+		Variable = new(ParseState)ir_variable(
+			Type,
+			ralloc_asprintf(ParseState, "in_%s", Semantic),
+			ir_var_in
+			);
+		Variable->read_only = true;
+		Variable->centroid = InputQualifier.Fields.bCentroid;
+		Variable->interpolation = InputQualifier.Fields.InterpolationMode;
+		Variable->is_patch_constant = InputQualifier.Fields.bIsPatchConstant;
 
-			Variable = new(ParseState)ir_variable(VariableType, ralloc_asprintf(ParseState, "in_%s", Semantic), ir_var_in);
-			Variable->read_only = true;
-			Variable->is_interface_block = true;
-			Variable->centroid = InputQualifier.Fields.bCentroid;
-			Variable->interpolation = InputQualifier.Fields.InterpolationMode;
-			Variable->is_patch_constant = InputQualifier.Fields.bIsPatchConstant;
-			DeclInstructions->push_tail(Variable);
-			ParseState->symbols->add_variable(Variable);
+		if (ParseState->bGenerateLayoutLocations && !InputQualifier.Fields.bIsPatchConstant)
+		{
+			ConfigureInOutVariableLayout(Frequency, ParseState, Semantic, Variable, ir_var_in);
+		}
 
-			if (ParseState->bGenerateLayoutLocations && !Variable->is_patch_constant)
-			{
-				ConfigureInOutVariableLayout(Frequency, ParseState, Semantic, Variable, ir_var_in);
-			}
+		DeclInstructions->push_tail(Variable);
+		ParseState->symbols->add_variable(Variable);
 
-			ir_rvalue* VariableDeref = new(ParseState)ir_dereference_variable(Variable);
-			if (SemanticArraySize)
-			{
-				// Deref inside array first
-				VariableDeref = new(ParseState)ir_dereference_array(VariableDeref, new(ParseState)ir_constant((unsigned)SemanticArrayIndex)
-					);
-			}
-			VariableDeref = new(ParseState)ir_dereference_record(VariableDeref, ralloc_strdup(ParseState, "Data"));
-		#endif
+		ir_dereference* VariableDeref = new(ParseState)ir_dereference_variable(Variable);
+		if (SemanticArraySize > 0)
+		{
+			// Deref inside array first
+			VariableDeref = new(ParseState) ir_dereference_array(VariableDeref, new(ParseState) ir_constant((unsigned)SemanticArrayIndex));
+		}
 
 		return VariableDeref;
 	}
@@ -4088,69 +4036,29 @@ static ir_rvalue* GenShaderOutputSemantic(
 
 	*DestVariableType = Type;
 
-	#if VULKAN_HLSLCC_GEN_GLSL_IN_OUT_WITHOUT_STRUCTURES
 	// This code-section replacces "layout(location=0) out struct { vec4 Data; } out_TEXCOORD0;" pattern to
-		// "layout(location=0) out vec4 out_TEXCOORD0;".
-
-		// Regular attribute
-		Variable = new(ParseState)ir_variable(
-			Type,
-			ralloc_asprintf(ParseState, "%s_%s", "out", Semantic),
-			ir_var_out
-			);
-		Variable->read_only = true;
-		Variable->centroid = OutputQualifier.Fields.bCentroid;
-		Variable->interpolation = OutputQualifier.Fields.InterpolationMode;
-		Variable->is_patch_constant = OutputQualifier.Fields.bIsPatchConstant;
-
-		if (ParseState->bGenerateLayoutLocations && !OutputQualifier.Fields.bIsPatchConstant)
-		{
-			ConfigureInOutVariableLayout(Frequency, ParseState, Semantic, Variable, ir_var_out);
-		}
-
-		DeclInstructions->push_tail(Variable);
-		ParseState->symbols->add_variable(Variable);
+	// "layout(location=0) out vec4 out_TEXCOORD0;".
 
-		ir_dereference_variable* VariableDeref = new(ParseState)ir_dereference_variable(Variable);
-	#else
-		// Create variable
-		glsl_struct_field *StructField = ralloc_array(ParseState, glsl_struct_field, 1);
-
-		memset(StructField, 0, sizeof(glsl_struct_field));
-		StructField[0].type = Type;
-		StructField[0].name = ralloc_strdup(ParseState, "Data");
-
-		const glsl_type* VariableType = glsl_type::get_record_instance(StructField, 1, ralloc_strdup(ParseState, Semantic));
-
-		if (Frequency == HSF_HullShader && !OutputQualifier.Fields.bIsPatchConstant)
-		{
-			VariableType = glsl_type::get_array_instance(VariableType, ParseState->tessellation.outputcontrolpoints);
-		}
-		Variable = new(ParseState)ir_variable(VariableType, ralloc_asprintf(ParseState, "out_%s", Semantic), ir_var_out);
-
-		Variable->centroid = OutputQualifier.Fields.bCentroid;
-		Variable->interpolation = OutputQualifier.Fields.InterpolationMode;
-		Variable->is_interface_block = true;
-		Variable->is_patch_constant = OutputQualifier.Fields.bIsPatchConstant;
-
-		if (ParseState->bGenerateLayoutLocations && !Variable->is_patch_constant)
-		{
-			ConfigureInOutVariableLayout(Frequency, ParseState, Semantic, Variable, ir_var_out);
-		}
-
-		DeclInstructions->push_tail(Variable);
-		ParseState->symbols->add_variable(Variable);
-
-		ir_rvalue* VariableDeref = new(ParseState)ir_dereference_variable(Variable);
+	// Regular attribute
+	Variable = new(ParseState)ir_variable(
+		Type,
+		ralloc_asprintf(ParseState, "%s_%s", "out", Semantic),
+		ir_var_out
+		);
+	Variable->read_only = true;
+	Variable->centroid = OutputQualifier.Fields.bCentroid;
+	Variable->interpolation = OutputQualifier.Fields.InterpolationMode;
+	Variable->is_patch_constant = OutputQualifier.Fields.bIsPatchConstant;
 
-		if (Frequency == HSF_HullShader && !OutputQualifier.Fields.bIsPatchConstant)
-		{
-			VariableDeref = new(ParseState)ir_dereference_array(VariableDeref, new(ParseState)ir_dereference_variable(ParseState->symbols->get_variable("gl_InvocationID")));
-		}
+	if (ParseState->bGenerateLayoutLocations && !OutputQualifier.Fields.bIsPatchConstant)
+	{
+		ConfigureInOutVariableLayout(Frequency, ParseState, Semantic, Variable, ir_var_out);
+	}
 
-		VariableDeref = new(ParseState)ir_dereference_record(VariableDeref, ralloc_strdup(ParseState, "Data"));
-	#endif
+	DeclInstructions->push_tail(Variable);
+	ParseState->symbols->add_variable(Variable);
 
+	ir_dereference_variable* VariableDeref = new(ParseState)ir_dereference_variable(Variable);
 	return VariableDeref;
 }
 
diff --git a/Engine/Source/Developer/Windows/ShaderFormatD3D/Private/D3D11ShaderCompiler.cpp b/Engine/Source/Developer/Windows/ShaderFormatD3D/Private/D3D11ShaderCompiler.cpp
--- a/Engine/Source/Developer/Windows/ShaderFormatD3D/Private/D3D11ShaderCompiler.cpp
+++ b/Engine/Source/Developer/Windows/ShaderFormatD3D/Private/D3D11ShaderCompiler.cpp
@@ -870,8 +870,14 @@ void CompileD3D11Shader(const FShaderCompilerInput& Input,FShaderCompilerOutput&
 		TArray<FString> UsedOutputs = Input.UsedOutputs;
 		UsedOutputs.AddUnique(TEXT("SV_POSITION"));
 
+		// We can't remove any of the output-only system semantics
+		//@todo - there are a bunch of tessellation ones as well
+		TArray<FString> Exceptions;
+		Exceptions.AddUnique(TEXT("SV_ClipDistance"));
+		Exceptions.AddUnique(TEXT("SV_CullDistance"));
+
 		TArray<FString> Errors;
-		if (!RemoveUnusedOutputs(PreprocessedShaderSource, UsedOutputs, EntryPointName, Errors))
+		if (!RemoveUnusedOutputs(PreprocessedShaderSource, UsedOutputs, Exceptions, EntryPointName, Errors))
 		{
 			UE_LOG(LogD3D11ShaderCompiler, Warning, TEXT("Failed to Remove unused outputs [%s]!"), *Input.DumpDebugInfoPath);
 			for (int32 Index = 0; Index < Errors.Num(); ++Index)
diff --git a/Engine/Source/Editor/LevelEditor/Private/LevelEditorActions.cpp b/Engine/Source/Editor/LevelEditor/Private/LevelEditorActions.cpp
--- a/Engine/Source/Editor/LevelEditor/Private/LevelEditorActions.cpp
+++ b/Engine/Source/Editor/LevelEditor/Private/LevelEditorActions.cpp
@@ -654,12 +654,9 @@ void FLevelEditorActionCallbacks::BuildLODsOnly_Execute()
 	FEditorBuildUtils::EditorBuild(GetWorld(), FBuildOptions::BuildHierarchicalLOD);
 }
 
-ENGINE_API void BuildTextureStreamingData(UWorld*);
-
 void FLevelEditorActionCallbacks::BuildTextureStreamingOnly_Execute()
 {
-	// Build TexCoord Scale Shaders
-	BuildTextureStreamingData(GetWorld());
+	FEditorBuildUtils::EditorBuildTextureStreaming(GetWorld());
 }
 
 bool FLevelEditorActionCallbacks::IsLightingQualityChecked( ELightingBuildQuality TestQuality )
diff --git a/Engine/Source/Editor/Persona/Private/PersonaMeshDetails.cpp b/Engine/Source/Editor/Persona/Private/PersonaMeshDetails.cpp
--- a/Engine/Source/Editor/Persona/Private/PersonaMeshDetails.cpp
+++ b/Engine/Source/Editor/Persona/Private/PersonaMeshDetails.cpp
@@ -249,8 +249,8 @@ void FSkelMeshReductionSettingsLayout::GenerateChildContent(IDetailChildrenBuild
 		.ValueContent()
 			[
 				SNew(SCheckBox)
-				.IsChecked(this, &FSkelMeshReductionSettingsLayout::ShouldRecalculateNormals)
-				.OnCheckStateChanged(this, &FSkelMeshReductionSettingsLayout::OnRecalculateNormalsChanged)
+				.IsChecked(this, &FSkelMeshReductionSettingsLayout::ShouldRecomputeTangents)
+				.OnCheckStateChanged(this, &FSkelMeshReductionSettingsLayout::OnRecomputeTangentsChanged)
 			];
 	}
 
@@ -377,7 +377,7 @@ float FSkelMeshReductionSettingsLayout::GetWeldingThreshold() const
 	return ReductionSettings.WeldingThreshold;
 }
 
-ECheckBoxState FSkelMeshReductionSettingsLayout::ShouldRecalculateNormals() const
+ECheckBoxState FSkelMeshReductionSettingsLayout::ShouldRecomputeTangents() const
 {
 	return ReductionSettings.bRecalcNormals ? ECheckBoxState::Checked : ECheckBoxState::Unchecked;
 }
@@ -413,7 +413,7 @@ void FSkelMeshReductionSettingsLayout::OnWeldingThresholdChanged(float NewValue)
 	ReductionSettings.WeldingThreshold = NewValue;
 }
 
-void FSkelMeshReductionSettingsLayout::OnRecalculateNormalsChanged(ECheckBoxState NewValue)
+void FSkelMeshReductionSettingsLayout::OnRecomputeTangentsChanged(ECheckBoxState NewValue)
 {
 	ReductionSettings.bRecalcNormals = NewValue == ECheckBoxState::Checked;
 }
@@ -1130,6 +1130,19 @@ TSharedRef<SWidget> FPersonaMeshDetails::OnGenerateCustomMaterialWidgetsForMater
 			]
 		]
 
+		+SVerticalBox::Slot()
+		.Padding(0,2,0,0)
+		[
+			SNew(SCheckBox)
+			.IsChecked(this, &FPersonaMeshDetails::IsRecomputeTangentEnabled, MaterialIndex)
+			.OnCheckStateChanged(this, &FPersonaMeshDetails::OnRecomputeTangentChanged, MaterialIndex)
+			[
+				SNew(STextBlock)
+				.Font(FEditorStyle::GetFontStyle("StaticMeshEditor.NormalFont"))
+				.Text(LOCTEXT("Recompute Tangent", "Recompute Tangent"))
+			]
+		]
+
 		+SVerticalBox::Slot()
 		.Padding(0,2,0,0)
 		.AutoHeight()
@@ -1314,6 +1327,51 @@ void FPersonaMeshDetails::OnShadowCastingChanged(ECheckBoxState NewState, int32
 	}
 }
 
+
+ECheckBoxState FPersonaMeshDetails::IsRecomputeTangentEnabled(int32 MaterialIndex) const
+{
+	ECheckBoxState State = ECheckBoxState::Unchecked;
+	const USkeletalMesh* Mesh = SkeletalMeshPtr.Get();
+	if (Mesh && MaterialIndex < Mesh->Materials.Num())
+	{
+		State = Mesh->Materials[MaterialIndex].bRecomputeTangent ? ECheckBoxState::Checked : ECheckBoxState::Unchecked;
+	}
+
+	return State;
+}
+
+void FPersonaMeshDetails::OnRecomputeTangentChanged(ECheckBoxState NewState, int32 MaterialIndex)
+{
+	USkeletalMesh* Mesh = SkeletalMeshPtr.Get();
+
+	if (Mesh)
+	{
+		if (NewState == ECheckBoxState::Checked)
+		{
+			const FScopedTransaction Transaction(LOCTEXT("SetRecomputeTangentFlag", "Set Recompute Tangent For Material"));
+			Mesh->Modify();
+			Mesh->Materials[MaterialIndex].bRecomputeTangent = true;
+		}
+		else if (NewState == ECheckBoxState::Unchecked)
+		{
+			const FScopedTransaction Transaction(LOCTEXT("ClearRecomputeTangentFlag", "Clear Recompute Tangent For Material"));
+			Mesh->Modify();
+			Mesh->Materials[MaterialIndex].bRecomputeTangent = false;
+		}
+		for (TObjectIterator<USkinnedMeshComponent> It; It; ++It)
+		{
+			USkinnedMeshComponent* MeshComponent = *It;
+			if (MeshComponent &&
+				!MeshComponent->IsTemplate() &&
+				MeshComponent->SkeletalMesh == Mesh)
+			{
+				MeshComponent->UpdateRecomputeTangent(MaterialIndex);
+				MeshComponent->MarkRenderStateDirty();
+			}
+		}
+		PersonaPtr->RefreshViewport();
+	}
+}
 int32 FPersonaMeshDetails::GetMaterialIndex(int32 LODIndex, int32 SectionIndex)
 {
 	USkeletalMesh* SkelMesh = PersonaPtr->GetMesh();
diff --git a/Engine/Source/Editor/Persona/Private/PersonaMeshDetails.h b/Engine/Source/Editor/Persona/Private/PersonaMeshDetails.h
--- a/Engine/Source/Editor/Persona/Private/PersonaMeshDetails.h
+++ b/Engine/Source/Editor/Persona/Private/PersonaMeshDetails.h
@@ -60,7 +60,7 @@ private:
 	float GetPercentTriangles() const;
 	float GetMaxDeviation() const;
 	float GetWeldingThreshold() const;
-	ECheckBoxState ShouldRecalculateNormals() const;
+	ECheckBoxState ShouldRecomputeTangents() const;
 	float GetHardAngleThreshold() const;
 	int32 GetMaxBonesPerVertex() const;
 	int32 GetBaseLOD() const;
@@ -68,7 +68,7 @@ private:
 	void OnPercentTrianglesChanged(float NewValue);
 	void OnMaxDeviationChanged(float NewValue);
 	void OnReductionAmountChanged(float NewValue);
-	void OnRecalculateNormalsChanged(ECheckBoxState NewValue);
+	void OnRecomputeTangentsChanged(ECheckBoxState NewValue);
 	void OnWeldingThresholdChanged(float NewValue);
 	void OnHardAngleThresholdChanged(float NewValue);
 	void OnMaxBonesPerVertexChanged(int32 NewValue);
@@ -170,7 +170,7 @@ private:
 	*/
 	void OnSectionIsolatedChanged(ECheckBoxState NewState, int32 SectionIndex);
 
-		/**
+	/**
 	 * Handler for check box display based on whether the material has shadow casting enabled
 	 *
 	 * @param MaterialIndex	The material index which a section in a specific LOD model has
@@ -184,6 +184,20 @@ private:
 	 */
 	void OnShadowCastingChanged(ECheckBoxState NewState, int32 MaterialIndex);
 
+	/**
+	* Handler for check box display based on whether this section does recalculate normal or not
+	*
+	* @param MaterialIndex	The material index which a section in a specific LOD model has
+	*/
+	ECheckBoxState IsRecomputeTangentEnabled(int32 MaterialIndex) const;
+
+	/**
+	* Handler for changing recalulate normal status on a material
+	*
+	* @param MaterialIndex	The material index which a section in a specific LOD model has
+	*/
+	void OnRecomputeTangentChanged(ECheckBoxState NewState, int32 MaterialIndex);
+
 	/**
 	 * Handler for enabling delete button on materials
 	 *
diff --git a/Engine/Source/Editor/PlacementMode/Private/PlacementModeModule.cpp b/Engine/Source/Editor/PlacementMode/Private/PlacementModeModule.cpp
--- a/Engine/Source/Editor/PlacementMode/Private/PlacementModeModule.cpp
+++ b/Engine/Source/Editor/PlacementMode/Private/PlacementModeModule.cpp
@@ -5,7 +5,7 @@
 #include "PlacementMode.h"
 #include "AssetToolsModule.h"
 #include "KismetEditorUtilities.h"
-
+#include "Classes/ActorFactories/ActorFactoryPlanarReflection.h"
 
 struct FPlacementCategory : FPlacementCategoryInfo
 {
@@ -162,6 +162,7 @@ public:
 			Category->Items.Add(CreateID(), MakeShareable( new FPlaceableItem(*UActorFactoryExponentialHeightFog::StaticClass(), SortOrder+=10)) );
 			Category->Items.Add(CreateID(), MakeShareable( new FPlaceableItem(*UActorFactorySphereReflectionCapture::StaticClass(), SortOrder+=10)) );
 			Category->Items.Add(CreateID(), MakeShareable( new FPlaceableItem(*UActorFactoryBoxReflectionCapture::StaticClass(), SortOrder+=10)) );
+			Category->Items.Add(CreateID(), MakeShareable( new FPlaceableItem(*UActorFactoryPlanarReflection::StaticClass(), SortOrder += 10)));
 			Category->Items.Add(CreateID(), MakeShareable( new FPlaceableItem(*UActorFactoryDeferredDecal::StaticClass(), SortOrder+=10)) );
 		}
 
diff --git a/Engine/Source/Editor/StaticMeshEditor/Private/SStaticMeshEditorViewport.cpp b/Engine/Source/Editor/StaticMeshEditor/Private/SStaticMeshEditorViewport.cpp
--- a/Engine/Source/Editor/StaticMeshEditor/Private/SStaticMeshEditorViewport.cpp
+++ b/Engine/Source/Editor/StaticMeshEditor/Private/SStaticMeshEditorViewport.cpp
@@ -281,7 +281,7 @@ void SStaticMeshEditorViewport::UpdatePreviewMesh(UStaticMesh* InStaticMesh)
 	PreviewMeshComponent->SetStaticMesh(InStaticMesh);
 
 	// Update streaming data for debug viewmode feedback
-	PreviewMeshComponent->UpdateStreamingTextureInfos(true);
+	PreviewMeshComponent->UpdateStreamingSectionData(FTexCoordScaleMap());
 
 	PreviewScene.AddComponent(PreviewMeshComponent,FTransform::Identity);
 
diff --git a/Engine/Source/Editor/UnrealEd/Classes/Editor/EditorEngine.h b/Engine/Source/Editor/UnrealEd/Classes/Editor/EditorEngine.h
--- a/Engine/Source/Editor/UnrealEd/Classes/Editor/EditorEngine.h
+++ b/Engine/Source/Editor/UnrealEd/Classes/Editor/EditorEngine.h
@@ -2128,7 +2128,7 @@ public:
 	float GetGridSize();
 
 	/** 
-	 * @return - if the grid size is part of the 1,2,48,16,.. list or not
+	 * @return - if the grid size is part of the 1,2,4,8,16,.. list or not
 	 */
 	bool IsGridSizePowerOfTwo() const;
 
diff --git a/Engine/Source/Editor/UnrealEd/Private/EditorBrushBuilder.cpp b/Engine/Source/Editor/UnrealEd/Private/EditorBrushBuilder.cpp
--- a/Engine/Source/Editor/UnrealEd/Private/EditorBrushBuilder.cpp
+++ b/Engine/Source/Editor/UnrealEd/Private/EditorBrushBuilder.cpp
@@ -404,17 +404,17 @@ void UCubeBuilder::BuildCube( int32 Direction, float dx, float dy, float dz, boo
 	if( _tessellated )
 	{
 		Poly3i(Direction,n+0,n+1,n+3);
-		Poly3i(Direction,n+0,n+3,n+2);
+		Poly3i(Direction,n+3,n+2,n+0);
 		Poly3i(Direction,n+2,n+3,n+7);
-		Poly3i(Direction,n+2,n+7,n+6);
+		Poly3i(Direction,n+7,n+6,n+2);
 		Poly3i(Direction,n+6,n+7,n+5);
-		Poly3i(Direction,n+6,n+5,n+4);
+		Poly3i(Direction,n+5,n+4,n+6);
 		Poly3i(Direction,n+4,n+5,n+1);
-		Poly3i(Direction,n+4,n+1,n+0);
+		Poly3i(Direction,n+1,n+0,n+4);
 		Poly3i(Direction,n+3,n+1,n+5);
-		Poly3i(Direction,n+3,n+5,n+7);
+		Poly3i(Direction,n+5,n+7,n+3);
 		Poly3i(Direction,n+0,n+2,n+6);
-		Poly3i(Direction,n+0,n+6,n+4);
+		Poly3i(Direction,n+6,n+4,n+0);
 	}
 	else
 	{
diff --git a/Engine/Source/Editor/UnrealEd/Private/EditorBuildUtils.cpp b/Engine/Source/Editor/UnrealEd/Private/EditorBuildUtils.cpp
--- a/Engine/Source/Editor/UnrealEd/Private/EditorBuildUtils.cpp
+++ b/Engine/Source/Editor/UnrealEd/Private/EditorBuildUtils.cpp
@@ -22,6 +22,7 @@
 #include "AI/Navigation/NavigationSystem.h"
 #include "HierarchicalLOD.h"
 #include "ActorEditorUtils.h"
+#include "MaterialUtilities.h"
 
 DEFINE_LOG_CATEGORY_STATIC(LogEditorBuildUtils, Log, All);
 
@@ -40,7 +41,7 @@ const FName FBuildOptions::BuildAll(TEXT("BuildAll"));
 const FName FBuildOptions::BuildAllSubmit(TEXT("BuildAllSubmit"));
 const FName FBuildOptions::BuildAllOnlySelectedPaths(TEXT("BuildAllOnlySelectedPaths"));
 const FName FBuildOptions::BuildHierarchicalLOD(TEXT("BuildHierarchicalLOD"));
-
+const FName FBuildOptions::BuildTextureStreaming(TEXT("BuildTextureStreaming"));
 
 bool FEditorBuildUtils::bBuildingNavigationFromUserRequest = false;
 TMap<FName, FEditorBuildUtils::FCustomBuildType> FEditorBuildUtils::CustomBuildTypes;
@@ -1028,4 +1029,34 @@ void FEditorBuildUtils::TriggerHierarchicalLODBuilder(UWorld* InWorld, FName Id)
 	InWorld->HierarchicalLODBuilder->BuildMeshesForLODActors();
 }
 
+bool FEditorBuildUtils::EditorBuildTextureStreaming(UWorld* InWorld)
+{
+
+	const EMaterialQualityLevel::Type QualityLevel = EMaterialQualityLevel::High;
+	const ERHIFeatureLevel::Type FeatureLevel = GMaxRHIFeatureLevel;
+
+	CollectGarbage( GARBAGE_COLLECTION_KEEPFLAGS );
+
+	FTexCoordScaleMap TexCoordScales;
+	BuildTextureStreamingShaders(InWorld, QualityLevel, FeatureLevel, TexCoordScales);
+
+	{
+		const double StartTime = FPlatformTime::Seconds();
+		for (FTexCoordScaleMap::TIterator It(TexCoordScales); It; ++It)
+		{
+			UMaterialInterface* MaterialInterface = It.Key();
+			TArray<FMaterialTexCoordBuildInfo>& Scales = It.Value();
+
+
+			FMaterialUtilities::ExportMaterialTexCoordScales(MaterialInterface, QualityLevel, FeatureLevel, Scales);
+		}
+		UE_LOG(LogLevel, Display, TEXT("Export Material TexCoord Scales took %.3f seconds."), FPlatformTime::Seconds() - StartTime);
+	}
+
+	BuildTextureStreamingData(InWorld, TexCoordScales, QualityLevel, FeatureLevel);
+	CollectGarbage( GARBAGE_COLLECTION_KEEPFLAGS );
+
+	return true;
+}
+
 #undef LOCTEXT_NAMESPACE
\ No newline at end of file
diff --git a/Engine/Source/Editor/UnrealEd/Private/Factories/ActorFactory.cpp b/Engine/Source/Editor/UnrealEd/Private/Factories/ActorFactory.cpp
--- a/Engine/Source/Editor/UnrealEd/Private/Factories/ActorFactory.cpp
+++ b/Engine/Source/Editor/UnrealEd/Private/Factories/ActorFactory.cpp
@@ -66,6 +66,8 @@ ActorFactory.cpp:
 #include "Animation/AnimBlueprintGeneratedClass.h"
 #include "Kismet2/ComponentEditorUtils.h"
 #include "Components/BillboardComponent.h"
+#include "Classes/ActorFactories/ActorFactoryPlanarReflection.h"
+#include "Engine/PlanarReflection.h"
 
 #include "LevelSequence.h"
 #include "LevelSequenceActor.h"
@@ -1601,6 +1603,18 @@ UActorFactoryBoxReflectionCapture::UActorFactoryBoxReflectionCapture(const FObje
 	bUseSurfaceOrientation = true;
 }
 
+/*-----------------------------------------------------------------------------
+UActorFactoryPlanarReflection
+-----------------------------------------------------------------------------*/
+UActorFactoryPlanarReflection::UActorFactoryPlanarReflection(const FObjectInitializer& ObjectInitializer)
+: Super(ObjectInitializer)
+{
+	DisplayName = LOCTEXT("PlanarReflectionDisplayName", "Planar Reflection");
+	NewActorClass = APlanarReflection::StaticClass();
+	SpawnPositionOffset = FVector(0, 0, 0);
+	bUseSurfaceOrientation = false;
+}
+
 /*-----------------------------------------------------------------------------
 UActorFactoryPlaneReflectionCapture
 -----------------------------------------------------------------------------*/
diff --git a/Engine/Source/Editor/UnrealEd/Private/Lightmass/Lightmass.cpp b/Engine/Source/Editor/UnrealEd/Private/Lightmass/Lightmass.cpp
--- a/Engine/Source/Editor/UnrealEd/Private/Lightmass/Lightmass.cpp
+++ b/Engine/Source/Editor/UnrealEd/Private/Lightmass/Lightmass.cpp
@@ -4063,7 +4063,7 @@ bool FLightmassProcessor::ImportTextureMapping(int32 Channel, FTextureMappingImp
 
 	if (TMImport.QuantizedData->Data.Num() > 0)
 	{
-		TMImport.UnmappedTexelsPercentage = (float)NumUnmappedTexels / (float)TMImport.QuantizedData->Data.Num();
+		TMImport.UnmappedTexelsPercentage = 100.0f * (float)NumUnmappedTexels / (float)TMImport.QuantizedData->Data.Num();
 	}
 	else
 	{
diff --git a/Engine/Source/Editor/UnrealEd/Private/PlayLevel.cpp b/Engine/Source/Editor/UnrealEd/Private/PlayLevel.cpp
--- a/Engine/Source/Editor/UnrealEd/Private/PlayLevel.cpp
+++ b/Engine/Source/Editor/UnrealEd/Private/PlayLevel.cpp
@@ -974,7 +974,7 @@ FString GenerateCmdLineForNextPieInstance(FIntPoint &WinPos, int32 &InstanceNum,
 		{
 			WindowBorderSize = TopLevelWindow->GetWindowBorderSize(true);
 		}
-
+		
 		// Get the size of the window based on the type
 		FIntPoint WinSize(0,0);
 		GetWindowSizeForInstanceType(WinSize, PlayInSettings);
@@ -1300,7 +1300,7 @@ void UEditorEngine::PlayStandaloneLocalPc(FString MapNameOverride, FIntPoint* Wi
 		GameNameOrProjectFile = FApp::GetGameName();
 	}
 
-	FString AdditionalParameters(TEXT(" -messaging -SessionName=\"Play in Standalone Game\""));
+	FString AdditionalParameters(TEXT(" -windowed -messaging -SessionName=\"Play in Standalone Game\""));
 	bool bRunningDebug = FParse::Param(FCommandLine::Get(), TEXT("debug"));
 	if (bRunningDebug)
 	{
diff --git a/Engine/Source/Editor/UnrealEd/Private/PreviewMaterial.cpp b/Engine/Source/Editor/UnrealEd/Private/PreviewMaterial.cpp
--- a/Engine/Source/Editor/UnrealEd/Private/PreviewMaterial.cpp
+++ b/Engine/Source/Editor/UnrealEd/Private/PreviewMaterial.cpp
@@ -47,7 +47,6 @@ public:
 	
 		}
 
-
 		{
 			bool bEditorStatsMaterial = Material->bIsMaterialEditorStatsMaterial;
 
@@ -63,9 +62,20 @@ public:
 				//cache for gpu skinned vertex factory if the material allows it
 				//this way we can have a preview skeletal mesh
 				if (bEditorStatsMaterial ||
-					!IsUsedWithSkeletalMesh() ||
-					(VertexFactoryType != FindVertexFactoryType(FName(TEXT("TGPUSkinVertexFactoryfalse"), FNAME_Find)) &&
-					VertexFactoryType != FindVertexFactoryType(FName(TEXT("TGPUSkinVertexFactorytrue"), FNAME_Find))))
+					!IsUsedWithSkeletalMesh())
+				{
+					return false;
+				}
+
+				extern ENGINE_API int32 GEnableGPUSkinCacheShaders;
+
+				bool bSkinCache = GEnableGPUSkinCacheShaders && (VertexFactoryType == FindVertexFactoryType(FName(TEXT("FGPUSkinPassthroughVertexFactory"), FNAME_Find)));
+					
+				if (
+					VertexFactoryType != FindVertexFactoryType(FName(TEXT("TGPUSkinVertexFactoryfalse"), FNAME_Find)) &&
+					VertexFactoryType != FindVertexFactoryType(FName(TEXT("TGPUSkinVertexFactorytrue"), FNAME_Find)) &&
+					!bSkinCache
+					)
 				{
 					return false;
 				}
diff --git a/Engine/Source/Editor/UnrealEd/Private/SEditorViewportViewMenu.cpp b/Engine/Source/Editor/UnrealEd/Private/SEditorViewportViewMenu.cpp
--- a/Engine/Source/Editor/UnrealEd/Private/SEditorViewportViewMenu.cpp
+++ b/Engine/Source/Editor/UnrealEd/Private/SEditorViewportViewMenu.cpp
@@ -282,14 +282,31 @@ TSharedRef<SWidget> SEditorViewportViewMenu::GenerateViewMenuContent() const
 						Menu.AddMenuEntry( BaseViewportCommands.LightmapDensityMode, NAME_None, LOCTEXT("LightmapDensityViewModeDisplayName", "Lightmap Density") );
 						Menu.AddMenuEntry( BaseViewportCommands.StationaryLightOverlapMode, NAME_None, LOCTEXT("StationaryLightOverlapViewModeDisplayName", "Stationary Light Overlap") );
 						Menu.AddMenuEntry( BaseViewportCommands.ShaderComplexityMode, NAME_None, LOCTEXT("ShaderComplexityViewModeDisplayName", "Shader Complexity") );
-						Menu.AddMenuEntry( BaseViewportCommands.ShaderComplexityWithQuadOverdrawMode, NAME_None, LOCTEXT("ShaderComplexityWithQuadOverdrawViewModeDisplayName", "Shader Complexity & Quads") );
-						Menu.AddMenuEntry( BaseViewportCommands.QuadOverdrawMode, NAME_None, LOCTEXT("QuadOverdrawViewModeDisplayName", "Quad Overdraw") );
+
+						if ( AllowDebugViewShaderMode(DVSM_ShaderComplexityBleedingQuadOverhead, GMaxRHIFeatureLevel) || AllowDebugViewShaderMode(DVSM_ShaderComplexityContainedQuadOverhead, GMaxRHIFeatureLevel) )
+						{
+							Menu.AddMenuEntry( BaseViewportCommands.ShaderComplexityWithQuadOverdrawMode, NAME_None, LOCTEXT("ShaderComplexityWithQuadOverdrawViewModeDisplayName", "Shader Complexity & Quads") );
+						}
+						if ( AllowDebugViewShaderMode(DVSM_QuadComplexity, GMaxRHIFeatureLevel) )
+						{
+							Menu.AddMenuEntry( BaseViewportCommands.QuadOverdrawMode, NAME_None, LOCTEXT("QuadOverdrawViewModeDisplayName", "Quad Overdraw") );
+						}
+
 						Menu.AddMenuEntry( BaseViewportCommands.LODColorationMode, NAME_None, LOCTEXT("LODColorationViewModeDisplayName", "LOD Coloration") );
 
 						Menu.BeginSection("TextureStreaming", LOCTEXT("TextureStreamingHeader", "Texture Streaming") );
-						Menu.AddMenuEntry( BaseViewportCommands.WantedMipsAccuracyMode, NAME_None, LOCTEXT("WantedMipsAccuracyDisplayName", "Wanted Mips Accuracy") );
-						Menu.AddMenuEntry( BaseViewportCommands.TexelFactorAccuracyMode, NAME_None, LOCTEXT("TexelFactorAccuracyDisplayName", "Texel Factor Accuracy") );
-						Menu.AddMenuEntry( BaseViewportCommands.TexCoordScaleAccuracyMode, NAME_None, LOCTEXT("TexCoordScaleAccuracyDisplayName", "TexCoord Scale Accuracy") );
+						if ( AllowDebugViewShaderMode(DVSM_WantedMipsAccuracy, GMaxRHIFeatureLevel) )
+						{
+							Menu.AddMenuEntry( BaseViewportCommands.WantedMipsAccuracyMode, NAME_None, LOCTEXT("WantedMipsAccuracyDisplayName", "Wanted Mips Accuracy") );
+						}
+						if ( AllowDebugViewShaderMode(DVSM_WantedMipsAccuracy, GMaxRHIFeatureLevel) )
+						{
+							Menu.AddMenuEntry( BaseViewportCommands.TexelFactorAccuracyMode, NAME_None, LOCTEXT("TexelFactorAccuracyDisplayName", "Texel Factor Accuracy") );
+						}
+						if ( AllowDebugViewShaderMode(DVSM_TexCoordScaleAccuracy, GMaxRHIFeatureLevel) )
+						{
+							Menu.AddMenuEntry( BaseViewportCommands.TexCoordScaleAccuracyMode, NAME_None, LOCTEXT("TexCoordScaleAccuracyDisplayName", "TexCoord Scale Accuracy") );
+						}
 						Menu.EndSection();
 					}
 				};
diff --git a/Engine/Source/Editor/UnrealEd/Public/EditorBuildUtils.h b/Engine/Source/Editor/UnrealEd/Public/EditorBuildUtils.h
--- a/Engine/Source/Editor/UnrealEd/Public/EditorBuildUtils.h
+++ b/Engine/Source/Editor/UnrealEd/Public/EditorBuildUtils.h
@@ -28,6 +28,8 @@ struct FBuildOptions
 	UNREALED_API static const FName BuildAllOnlySelectedPaths;
 	/** Build Hierarchical LOD system - need WorldSetting setup*/
 	UNREALED_API static const FName BuildHierarchicalLOD;
+	/** Build Hierarchical LOD system - need WorldSetting setup*/
+	UNREALED_API static const FName BuildTextureStreaming;
 };
 /**
  * Result of a custom editor build.
@@ -118,6 +120,16 @@ public:
 	 */
 	static UNREALED_API bool EditorBuild( UWorld* InWorld, FName Id, const bool bAllowLightingDialog = true );
 
+	/**
+	 * Perform an editor build for texture streaming
+	 *
+	 * @param	InWorld				WorldContext
+	 *
+	 * @return	true if the build completed successfully; false if it did not (or was manually canceled)
+	 */
+	static UNREALED_API bool EditorBuildTextureStreaming( UWorld* InWorld );
+
+
 	/** 
 	* check if navigation build was was triggered from editor as user request
 	*
diff --git a/Engine/Source/Programs/ShaderCompileWorker/Private/ShaderCompileWorker.cpp b/Engine/Source/Programs/ShaderCompileWorker/Private/ShaderCompileWorker.cpp
--- a/Engine/Source/Programs/ShaderCompileWorker/Private/ShaderCompileWorker.cpp
+++ b/Engine/Source/Programs/ShaderCompileWorker/Private/ShaderCompileWorker.cpp
@@ -594,56 +594,58 @@ static void CompileDirect(const TArray<const class IShaderFormat*>& ShaderFormat
 	FCommandLine::Parse(FCommandLine::Get(), Tokens, Switches);
 
 	FString InputFile;
-	if (Tokens.Num() < 1)
-	{
-		return;
-	}
-
-	for (int32 Index = 0; Index < Tokens.Num(); ++Index)
-	{
-		if (!Switches.Contains(Tokens[Index]))
-		{
-			InputFile = Tokens[Index];
-			break;
-		}
-	}
 
 	FName FormatName;
 	FString Entry = TEXT("Main");
+	bool bPipeline = false;
 	EShaderFrequency Frequency = SF_Pixel;
-	for (const FString& Switch : Switches)
+	for (const FString& Token : Tokens)
 	{
-		if (Switch.StartsWith(TEXT("format=")))
+		if (Switches.Contains(Token))
 		{
-			FormatName = FName(*Switch.RightChop(7));
-		}
-		else if (Switch.StartsWith(TEXT("entry=")))
-		{
-			Entry = Switch.RightChop(6);
-		}
-		else if (!FCString::Strcmp(*Switch, TEXT("ps")))
-		{
-			Frequency = SF_Pixel;
-		}
-		else if (!FCString::Strcmp(*Switch, TEXT("vs")))
-		{
-			Frequency = SF_Vertex;
-		}
-		else if (!FCString::Strcmp(*Switch, TEXT("hs")))
-		{
-			Frequency = SF_Hull;
-		}
-		else if (!FCString::Strcmp(*Switch, TEXT("ds")))
-		{
-			Frequency = SF_Domain;
-		}
-		else if (!FCString::Strcmp(*Switch, TEXT("gs")))
-		{
-			Frequency = SF_Geometry;
+			if (Token.StartsWith(TEXT("format=")))
+			{
+				FormatName = FName(*Token.RightChop(7));
+			}
+			else if (Token.StartsWith(TEXT("entry=")))
+			{
+				Entry = Token.RightChop(6);
+			}
+			else if (!FCString::Strcmp(*Token, TEXT("ps")))
+			{
+				Frequency = SF_Pixel;
+			}
+			else if (!FCString::Strcmp(*Token, TEXT("vs")))
+			{
+				Frequency = SF_Vertex;
+			}
+			else if (!FCString::Strcmp(*Token, TEXT("hs")))
+			{
+				Frequency = SF_Hull;
+			}
+			else if (!FCString::Strcmp(*Token, TEXT("ds")))
+			{
+				Frequency = SF_Domain;
+			}
+			else if (!FCString::Strcmp(*Token, TEXT("gs")))
+			{
+				Frequency = SF_Geometry;
+			}
+			else if (!FCString::Strcmp(*Token, TEXT("cs")))
+			{
+				Frequency = SF_Compute;
+			}
+			else if (!FCString::Strcmp(*Token, TEXT("pipeline")))
+			{
+				bPipeline = true;
+			}
 		}
-		else if (!FCString::Strcmp(*Switch, TEXT("cs")))
+		else
 		{
-			Frequency = SF_Compute;
+			if (InputFile.Len() == 0)
+			{
+				InputFile = Token;
+			}
 		}
 	}
 
@@ -657,6 +659,8 @@ static void CompileDirect(const TArray<const class IShaderFormat*>& ShaderFormat
 	Input.Target.Frequency = Frequency;
 	Input.bSkipPreprocessedCache = true;
 
+	Input.bCompilingForShaderPipeline = bPipeline;
+
 	FShaderCompilerOutput Output;
 
 	for (const IShaderFormat* Format : ShaderFormats)
diff --git a/Engine/Source/Runtime/Core/Private/HAL/ConsoleManager.cpp b/Engine/Source/Runtime/Core/Private/HAL/ConsoleManager.cpp
--- a/Engine/Source/Runtime/Core/Private/HAL/ConsoleManager.cpp
+++ b/Engine/Source/Runtime/Core/Private/HAL/ConsoleManager.cpp
@@ -1849,34 +1849,6 @@ static TAutoConsoleVariable<int32> CVarNumBufferedOcclusionQueries(
 	TEXT("More frames reduces the chance of stalling the CPU waiting for results, but increases out of date query artifacts."),
 	ECVF_RenderThreadSafe);
 
-static TAutoConsoleVariable<int32> CVarDistField(
-	TEXT("r.GenerateMeshDistanceFields"),
-	0,	
-	TEXT("Whether to build distance fields of static meshes, needed for distance field AO, which is used to implement Movable SkyLight shadows.\n")
-	TEXT("Enabling will increase mesh build times and memory usage.  Changing this value will cause a rebuild of all static meshes."),
-	ECVF_ReadOnly);
-
-static TAutoConsoleVariable<int32> CVarDistFieldRes(
-	TEXT("r.DistanceFields.MaxPerMeshResolution"),
-	128,	
-	TEXT("Highest resolution (in one dimension) allowed for a single static mesh asset, used to cap the memory usage of meshes with a large scale.\n")
-	TEXT("Changing this will cause all distance fields to be rebuilt.  Large values such as 512 can consume memory very quickly! (128Mb for one asset at 512)"),
-	ECVF_ReadOnly);
-
-static TAutoConsoleVariable<float> CVarDistFieldResScale(
-	TEXT("r.DistanceFields.DefaultVoxelDensity"),
-	.1f,	
-	TEXT("Determines how the default scale of a mesh converts into distance field voxel dimensions.\n")
-	TEXT("Changing this will cause all distance fields to be rebuilt.  Large values can consume memory very quickly!"),
-	ECVF_ReadOnly);
-
-static TAutoConsoleVariable<int32> CVarLandscapeGI(
-	TEXT("r.GenerateLandscapeGIData"),
-	1,
-	TEXT("Whether to generate a low-resolution base color texture for landscapes for rendering real-time global illumination.\n")
-	TEXT("This feature requires GenerateMeshDistanceFields is also enabled, and will increase mesh build times and memory usage.\n"),
-	ECVF_Default);
-
 static TAutoConsoleVariable<int32> CVarMinLogVerbosity(
 	TEXT("con.MinLogVerbosity"),
 	0,
diff --git a/Engine/Source/Runtime/Core/Private/Mac/CocoaWindow.cpp b/Engine/Source/Runtime/Core/Private/Mac/CocoaWindow.cpp
--- a/Engine/Source/Runtime/Core/Private/Mac/CocoaWindow.cpp
+++ b/Engine/Source/Runtime/Core/Private/Mac/CocoaWindow.cpp
@@ -212,7 +212,7 @@ NSString* NSPerformDragOperation = @"NSPerformDragOperation";
 		self.PreFullScreenRect.size = [self openGLFrame].size;
 		
 		// Use the current default fullscreen mode when switching via the OS button
-		auto CVar = IConsoleManager::Get().FindTConsoleVariableDataInt(TEXT("r.FullScreenMode"));
+		static auto CVar = IConsoleManager::Get().FindTConsoleVariableDataInt(TEXT("r.FullScreenMode"));
 		check(CVar);
 		self.TargetWindowMode = CVar->GetValueOnAnyThread() == 0 ? EWindowMode::Fullscreen : EWindowMode::WindowedFullscreen;
 		
diff --git a/Engine/Source/Runtime/Core/Private/Math/Sphere.cpp b/Engine/Source/Runtime/Core/Private/Math/Sphere.cpp
--- a/Engine/Source/Runtime/Core/Private/Math/Sphere.cpp
+++ b/Engine/Source/Runtime/Core/Private/Math/Sphere.cpp
@@ -66,7 +66,7 @@ FSphere FSphere::TransformBy(const FTransform& M) const
 
 float FSphere::GetVolume() const
 {
-	return 1.25f * PI * FMath::Pow(W, 3);
+	return (4.f / 3.f) * PI * (W * W * W);
 }
 
 FSphere& FSphere::operator+=(const FSphere &Other)
diff --git a/Engine/Source/Runtime/Core/Private/Misc/Core.cpp b/Engine/Source/Runtime/Core/Private/Misc/Core.cpp
--- a/Engine/Source/Runtime/Core/Private/Misc/Core.cpp
+++ b/Engine/Source/Runtime/Core/Private/Misc/Core.cpp
@@ -211,7 +211,7 @@ uint64					GFrameCounter					= 0;
 uint64					GLastGCFrame					= 0;
 /** Incremented once per frame before the scene is being rendered. In split screen mode this is incremented once for all views (not for each view). */
 uint32					GFrameNumber					= 1;
-/** Render Thread copy of the frame number. */
+/** NEED TO RENAME, for RT version of GFrameTime use View.ViewFamily->FrameNumber or pass down from RT from GFrameTime). */
 uint32					GFrameNumberRenderThread		= 1;
 #if !(UE_BUILD_SHIPPING && WITH_EDITOR)
 // We cannot count on this variable to be accurate in a shipped game, so make sure no code tries to use it
diff --git a/Engine/Source/Runtime/CoreUObject/Private/Serialization/AsyncLoading.cpp b/Engine/Source/Runtime/CoreUObject/Private/Serialization/AsyncLoading.cpp
--- a/Engine/Source/Runtime/CoreUObject/Private/Serialization/AsyncLoading.cpp
+++ b/Engine/Source/Runtime/CoreUObject/Private/Serialization/AsyncLoading.cpp
@@ -270,6 +270,22 @@ static FAutoConsoleVariableRef CVarPreloadPackageDependencies(
 	ECVF_Default
 	);
 
+uint32 FAsyncLoadingThread::AsyncLoadingThreadID = 0;
+
+static void IsTimeLimitExceededPrint(double InTickStartTime, double CurrentTime, float InTimeLimit, const TCHAR* InLastTypeOfWorkPerformed = nullptr, UObject* InLastObjectWorkWasPerformedOn = nullptr)
+{
+	// Log single operations that take longer than time limit (but only in cooked builds)
+	if (
+		(CurrentTime - InTickStartTime) > GTimeLimitExceededMinTime &&
+		(CurrentTime - InTickStartTime) > (GTimeLimitExceededMultiplier * InTimeLimit))
+	{
+		UE_LOG(LogStreaming, Warning, TEXT("IsTimeLimitExceeded: %s %s took (less than) %5.2f ms"),
+			InLastTypeOfWorkPerformed ? InLastTypeOfWorkPerformed : TEXT("unknown"),
+			InLastObjectWorkWasPerformedOn ? *InLastObjectWorkWasPerformedOn->GetFullName() : TEXT("nullptr"),
+			(CurrentTime - InTickStartTime) * 1000);
+	}
+}
+
 static FORCEINLINE bool IsTimeLimitExceeded(double InTickStartTime, bool bUseTimeLimit, float InTimeLimit, const TCHAR* InLastTypeOfWorkPerformed = nullptr, UObject* InLastObjectWorkWasPerformedOn = nullptr)
 {
 	bool bTimeLimitExceeded = false;
@@ -278,21 +294,17 @@ static FORCEINLINE bool IsTimeLimitExceeded(double InTickStartTime, bool bUseTim
 		double CurrentTime = FPlatformTime::Seconds();
 		bTimeLimitExceeded = CurrentTime - InTickStartTime > InTimeLimit;
 
-		// Log single operations that take longer than time limit (but only in cooked builds)
-		if (GWarnIfTimeLimitExceeded && 
-			(CurrentTime - InTickStartTime) > GTimeLimitExceededMinTime &&
-			(CurrentTime - InTickStartTime) > (GTimeLimitExceededMultiplier * InTimeLimit))
-		{			
-			UE_LOG(LogStreaming, Warning, TEXT("IsTimeLimitExceeded: %s %s took (less than) %5.2f ms"),
-				InLastTypeOfWorkPerformed ? InLastTypeOfWorkPerformed : TEXT("unknown"),
-				InLastObjectWorkWasPerformedOn ? *InLastObjectWorkWasPerformedOn->GetFullName() : TEXT("nullptr"),
-				(CurrentTime - InTickStartTime) * 1000);
+		if (GWarnIfTimeLimitExceeded)
+		{
+			IsTimeLimitExceededPrint(InTickStartTime, CurrentTime, InTimeLimit, InLastTypeOfWorkPerformed, InLastObjectWorkWasPerformedOn);
 		}
 	}
 	return bTimeLimitExceeded;
 }
-
-uint32 FAsyncLoadingThread::AsyncLoadingThreadID = 0;
+FORCEINLINE bool FAsyncPackage::IsTimeLimitExceeded()
+{
+	return AsyncLoadingThread.IsAsyncLoadingSuspended() || ::IsTimeLimitExceeded(TickStartTime, bUseTimeLimit, TimeLimit, LastTypeOfWorkPerformed, LastObjectWorkWasPerformedOn);
+}
 
 #if LOOKING_FOR_PERF_ISSUES
 FThreadSafeCounter FAsyncLoadingThread::BlockingCycles = 0;
@@ -1153,16 +1165,6 @@ void FAsyncPackage::DetachLinker()
 	}
 }
 
-/**
- * Returns whether time limit has been exceeded.
- *
- * @return true if time limit has been exceeded (and is used), false otherwise
- */
-bool FAsyncPackage::IsTimeLimitExceeded()
-{
-	return AsyncLoadingThread.IsAsyncLoadingSuspended() || ::IsTimeLimitExceeded(TickStartTime, bUseTimeLimit, TimeLimit, LastTypeOfWorkPerformed, LastObjectWorkWasPerformedOn);
-}
-
 /**
  * Gives up time slice if time limit is enabled.
  *
diff --git a/Engine/Source/Runtime/Engine/Classes/Components/PrimitiveComponent.h b/Engine/Source/Runtime/Engine/Classes/Components/PrimitiveComponent.h
--- a/Engine/Source/Runtime/Engine/Classes/Components/PrimitiveComponent.h
+++ b/Engine/Source/Runtime/Engine/Classes/Components/PrimitiveComponent.h
@@ -8,7 +8,6 @@
 #include "SceneTypes.h"
 #include "Engine/EngineTypes.h"
 #include "AI/Navigation/NavRelevantInterface.h"
-
 #include "PrimitiveComponent.generated.h"
 
 class FPrimitiveSceneProxy;
@@ -17,29 +16,8 @@ class UTexture;
 struct FEngineShowFlags;
 struct FConvexVolume;
 struct FNavigableGeometryExport;
-
-/** Information about a streaming texture that a primitive uses for rendering. */
-USTRUCT()
-struct FStreamingTexturePrimitiveInfo
-{
-	GENERATED_USTRUCT_BODY()
-
-	FStreamingTexturePrimitiveInfo()
-		: Texture(nullptr)
-		, Bounds(ForceInit)
-		, TexelFactor(1.0f)
-	{
-	}
-
-	UPROPERTY()
-	UTexture* Texture;
-
-	UPROPERTY()
-	FBoxSphereBounds Bounds;
-
-	UPROPERTY()
-	float TexelFactor;
-};
+struct FStreamingTexturePrimitiveInfo;
+class FStreamingTextureLevelContext;
 
 /** Determines whether a Character can attempt to step up onto a component when they walk in to it. */
 UENUM()
@@ -290,10 +268,6 @@ public:
 	UPROPERTY(EditAnywhere, BlueprintReadOnly, Category=Lighting, AdvancedDisplay, meta=(EditCondition="CastShadow", DisplayName = "Static Shadow"))
 	uint32 bCastStaticShadow:1;
 
-	/** Whether the object should be rendered into the planar refelction pass. */
-	UPROPERTY(EditAnywhere, Category=Lighting)
-	uint32 bVisibleInPlanarReflection:1;
-
 	/** 
 	 * Whether the object should cast a volumetric translucent shadow.
 	 * Volumetric translucent shadows are useful for primitives with smoothly changing opacity like particles representing a volume, 
@@ -1186,17 +1160,17 @@ public:
 
 	/**
 	 * Enumerates the streaming textures used by the primitive.
+	 * @param LevelContext - Level scope context used to process texture streaming build data.
 	 * @param OutStreamingTextures - Upon return, contains a list of the streaming textures used by the primitive.
 	 */
-	virtual void GetStreamingTextureInfo(TArray<struct FStreamingTexturePrimitiveInfo>& OutStreamingTextures) const
-	{
-	}
+	virtual void GetStreamingTextureInfo(FStreamingTextureLevelContext& LevelContext, TArray<FStreamingTexturePrimitiveInfo>& OutStreamingTextures) const {}
 
 	/**
 	 * Call GetStreamingTextureInfo and remove the elements with a NULL texture
 	 * @param OutStreamingTextures - Upon return, contains a list of the non-null streaming textures used by the primitive.
 	 */
-	void GetStreamingTextureInfoWithNULLRemoval(TArray<struct FStreamingTexturePrimitiveInfo>& OutStreamingTextures) const;
+	void GetStreamingTextureInfoWithNULLRemoval(FStreamingTextureLevelContext& LevelContext, TArray<FStreamingTexturePrimitiveInfo>& OutStreamingTextures) const;
+
 
 	/**
 	 * Determines the DPG the primitive's primary elements are drawn in.
diff --git a/Engine/Source/Runtime/Engine/Classes/Components/SceneCaptureComponent2D.h b/Engine/Source/Runtime/Engine/Classes/Components/SceneCaptureComponent2D.h
--- a/Engine/Source/Runtime/Engine/Classes/Components/SceneCaptureComponent2D.h
+++ b/Engine/Source/Runtime/Engine/Classes/Components/SceneCaptureComponent2D.h
@@ -33,15 +33,6 @@ class USceneCaptureComponent2D : public USceneCaptureComponent
 	UPROPERTY(interp, Category=SceneCapture, meta=(DisplayName = "Capture Source"))
 	TEnumAsByte<enum ESceneCaptureSource> CaptureSource;
 
-	UPROPERTY(EditAnywhere, Category = SceneCapture)
-	bool bIsPlanarReflection;
-
-	UPROPERTY(EditAnywhere, BlueprintReadWrite, Category = SceneCapture)
-	float ReflectionPlaneHeight;
-
-	UPROPERTY(EditAnywhere, BlueprintReadWrite, Category = SceneCapture)
-	FVector ReflectionPlaneNormal;
-
 	UPROPERTY(interp, Category=PostProcessVolume, meta=(ShowOnlyInnerProperties))
 	struct FPostProcessSettings PostProcessSettings;
 
diff --git a/Engine/Source/Runtime/Engine/Classes/Components/SkinnedMeshComponent.h b/Engine/Source/Runtime/Engine/Classes/Components/SkinnedMeshComponent.h
--- a/Engine/Source/Runtime/Engine/Classes/Components/SkinnedMeshComponent.h
+++ b/Engine/Source/Runtime/Engine/Classes/Components/SkinnedMeshComponent.h
@@ -203,6 +203,9 @@ public:
 
 	const TArray<int32>& GetMasterBoneMap() const { return MasterBoneMap; }
 
+	/** update Recalculate Normal flag in matching section */
+	void UpdateRecomputeTangent(int32 MaterialIndex);
+
 	/** 
 	 * When true, we will just using the bounds from our MasterPoseComponent.  This is useful for when we have a Mesh Parented
 	 * to the main SkelMesh (e.g. outline mesh or a full body overdraw effect that is toggled) that is always going to be the same
@@ -505,7 +508,7 @@ public:
 	virtual UMaterialInterface* GetMaterial(int32 MaterialIndex) const override;
 	virtual FPrimitiveSceneProxy* CreateSceneProxy() override;
 	virtual void GetUsedMaterials(TArray<UMaterialInterface*>& OutMaterials) const override;
-	virtual void GetStreamingTextureInfo(TArray<FStreamingTexturePrimitiveInfo>& OutStreamingTextures) const override;
+	virtual void GetStreamingTextureInfo(FStreamingTextureLevelContext& LevelContext, TArray<FStreamingTexturePrimitiveInfo>& OutStreamingTextures) const override;
 	virtual int32 GetNumMaterials() const override;
 	//~ End UPrimitiveComponent Interface
 
diff --git a/Engine/Source/Runtime/Engine/Classes/Components/SplineMeshComponent.h b/Engine/Source/Runtime/Engine/Classes/Components/SplineMeshComponent.h
--- a/Engine/Source/Runtime/Engine/Classes/Components/SplineMeshComponent.h
+++ b/Engine/Source/Runtime/Engine/Classes/Components/SplineMeshComponent.h
@@ -328,7 +328,7 @@ class ENGINE_API USplineMeshComponent : public UStaticMeshComponent, public IInt
 	/** Returns a vector which, when componentwise-multiplied by another vector, will zero all the components not corresponding to the supplied ESplineMeshAxis */
 	inline static FVector GetAxisMask(ESplineMeshAxis::Type InAxis);
 
-	virtual bool GetStreamingTextureFactors(float& OutWorldTexelFactor, float& OutWorldLightmapFactor, FBoxSphereBounds& OutBounds, int32 LODIndex, int32 ElementIndex) const override;
+	virtual bool GetStreamingTextureFactors(float& OutWorldTexelFactor, float& OutWorldLightmapFactor) const override;
 };
 
 const float& USplineMeshComponent::GetAxisValue(const FVector& InVector, ESplineMeshAxis::Type InAxis)
diff --git a/Engine/Source/Runtime/Engine/Classes/Components/StaticMeshComponent.h b/Engine/Source/Runtime/Engine/Classes/Components/StaticMeshComponent.h
--- a/Engine/Source/Runtime/Engine/Classes/Components/StaticMeshComponent.h
+++ b/Engine/Source/Runtime/Engine/Classes/Components/StaticMeshComponent.h
@@ -3,6 +3,7 @@
 #pragma once
 
 #include "SceneTypes.h"
+#include "Engine/TextureStreamingTypes.h"
 #include "Components/MeshComponent.h"
 #include "Runtime/RenderCore/Public/PackedNormal.h"
 #include "RawIndexBuffer.h"
@@ -14,6 +15,7 @@ class FStaticMeshStaticLightingMesh;
 class ULightComponent;
 struct FEngineShowFlags;
 struct FConvexVolume;
+struct FStreamingTextureBuildInfo;
 
 /** Cached vertex information at the time the mesh was painted. */
 USTRUCT()
@@ -237,11 +239,18 @@ class ENGINE_API UStaticMeshComponent : public UMeshComponent
 	UPROPERTY(transient)
 	TArray<struct FStaticMeshComponentLODInfo> LODData;
 
-	/** The streaming info for each LOD / Section. Needs to be persistent in order to be accessible when loading cooked maps. */
-	UPROPERTY()
-	TArray<FStreamingTexturePrimitiveInfo> StreamingTextureInfos;
+	/** The list of texture, bounds and scales. As computed in the texture streaming build process. */
+	UPROPERTY(duplicatetransient, NonTransactional)
+	TArray<FStreamingTextureBuildInfo> StreamingTextureData;
 
 #if WITH_EDITORONLY_DATA
+	/** 
+	 * Temporary section data used in the texture streaming build. 
+	 * Stays persistent to allow texture streaming accuracy view mode to inspect it.
+	 * The shared ptr is used to allow a safe way for the proxy to access it without duplicating it.
+	 */
+	TSharedPtr<TArray<FStreamingSectionBuildInfo>, ESPMode::NotThreadSafe> StreamingSectionData;
+
 	/** Derived data key of the static mesh, used to determine if an update from the source static mesh is required. */
 	UPROPERTY()
 	FString StaticMeshDerivedDataKey;
@@ -327,9 +336,27 @@ public:
 		return LightmassSettings.bShadowIndirectOnly;
 	}
 	virtual ELightMapInteractionType GetStaticLightingType() const override;
-	void UpdateStreamingTextureInfos(bool bForce = false);
-	virtual bool GetStreamingTextureFactors(float& OutWorldTexelFactor, float& OutWorldLightmapFactor, FBoxSphereBounds& OutBounds, int32 LODIndex, int32 ElementIndex) const;
-	virtual void GetStreamingTextureInfo(TArray<FStreamingTexturePrimitiveInfo>& OutStreamingTextures) const override;
+
+	/**
+	 *	Update the precomputed streaming data of this component.
+	 *
+	 *	@param	LevelTextures	[in,out]	The list of textures referred by all component of a level. The array index maps to UTexture2D::LevelIndex.
+	 *	@param	TexCoordScales	[in]		The texcoord scales for each texture register of each relevant materials.
+	 *	@param	QualityLevel	[in]		The quality level being used in the texture streaming build.
+	 *	@param	FeatureLevel	[in]		The feature level being used in the texture streaming build.
+	 */
+	virtual void UpdateStreamingTextureData(TArray<UTexture2D*>& LevelTextures, const FTexCoordScaleMap& TexCoordScales, EMaterialQualityLevel::Type QualityLevel, ERHIFeatureLevel::Type FeatureLevel);
+
+	/**
+	*	Update the precomputed streaming debug data of this component.
+	*
+	*	@param	TexCoordScales				The texcoord scales for each texture register of each relevant materials.
+	*/
+	virtual void UpdateStreamingSectionData(const FTexCoordScaleMap& TexCoordScales);
+
+	virtual bool GetStreamingTextureFactors(float& OutWorldTexelFactor, float& OutWorldLightmapFactor) const;
+	virtual void GetStreamingTextureInfo(FStreamingTextureLevelContext& LevelContext, TArray<FStreamingTexturePrimitiveInfo>& OutStreamingTextures) const override;
+
 	virtual class UBodySetup* GetBodySetup() override;
 	virtual bool CanEditSimulatePhysics() override;
 	virtual FPrimitiveSceneProxy* CreateSceneProxy() override;
diff --git a/Engine/Source/Runtime/Engine/Classes/Engine/Level.h b/Engine/Source/Runtime/Engine/Classes/Engine/Level.h
--- a/Engine/Source/Runtime/Engine/Classes/Engine/Level.h
+++ b/Engine/Source/Runtime/Engine/Classes/Engine/Level.h
@@ -394,6 +394,10 @@ public:
 	UPROPERTY()
 	TArray<FVector> StaticNavigableGeometry;
 
+	/** The Guid of each texture refered by FStreamingTextureBuildInfo::TextureLevelIndex	*/
+	UPROPERTY()
+	TArray<FGuid> StreamingTextureGuids;
+
 	/** Static information used by texture streaming code, generated during PreSave									*/
 	TMap<UTexture2D*,TArray<FStreamableTextureInstance> >	TextureToInstancesMap;
 
@@ -433,10 +437,14 @@ public:
 	/** Has texture streaming been built */
 	uint32										bTextureStreamingBuilt:1;
 
+	/** Whether a level transform rotation was applied since the texture streaming builds. Invalidates the precomputed streaming bounds. */
+	UPROPERTY()
+	uint32 										bTextureStreamingRotationChanged : 1;
+
 	/** Whether the level is currently visible/ associated with the world */
 	UPROPERTY(transient)
-	uint32										bIsVisible:1;
-	
+	uint32										bIsVisible:1;	
+
 	/** Whether this level is locked; that is, its actors are read-only 
 	 *	Used by WorldBrowser to lock a level when corresponding ULevelStreaming does not exist
 	 */
diff --git a/Engine/Source/Runtime/Engine/Classes/Engine/PlanarReflection.h b/Engine/Source/Runtime/Engine/Classes/Engine/PlanarReflection.h
new file mode 100644
--- /dev/null
+++ b/Engine/Source/Runtime/Engine/Classes/Engine/PlanarReflection.h
@@ -0,0 +1,42 @@
+// Copyright 1998-2016 Epic Games, Inc. All Rights Reserved.
+
+/**
+ * APlanarReflection
+ */
+
+#pragma once
+#include "SceneCapture.h"
+#include "PlanarReflection.generated.h"
+
+UCLASS(hidecategories=(Collision, Material, Attachment, Actor), MinimalAPI)
+class APlanarReflection : public ASceneCapture
+{
+	GENERATED_UCLASS_BODY()
+
+private_subobject:
+	/** Planar reflection component. */
+	UPROPERTY(Category = SceneCapture, VisibleAnywhere, BlueprintReadOnly, meta = (AllowPrivateAccess = "true"))
+	class UPlanarReflectionComponent* PlanarReflectionComponent;
+
+public:
+
+	//~ Begin AActor Interface
+	ENGINE_API virtual void PostActorCreated() override;
+
+#if WITH_EDITOR
+	virtual void EditorApplyScale(const FVector& DeltaScale, const FVector* PivotLocation, bool bAltDown, bool bShiftDown, bool bCtrlDown) override;
+#endif
+	//~ End AActor Interface.
+
+	UFUNCTION(BlueprintCallable, Category="Rendering")
+	void OnInterpToggle(bool bEnable);
+
+	/** Returns subobject **/
+	ENGINE_API class UPlanarReflectionComponent* GetPlanarReflectionComponent() const
+	{
+		return PlanarReflectionComponent;
+	}
+};
+
+
+
diff --git a/Engine/Source/Runtime/Engine/Classes/Engine/RendererSettings.h b/Engine/Source/Runtime/Engine/Classes/Engine/RendererSettings.h
--- a/Engine/Source/Runtime/Engine/Classes/Engine/RendererSettings.h
+++ b/Engine/Source/Runtime/Engine/Classes/Engine/RendererSettings.h
@@ -285,6 +285,12 @@ class ENGINE_API URendererSettings : public UDeveloperSettings
 		ConfigRestartRequired=true))
 	uint32 bSelectiveBasePassOutputs:1;
 
+	UPROPERTY(config, EditAnywhere, Category = Optimizations, meta = (
+		ConsoleVariable = "r.AllowGlobalClipPlane", DisplayName = "Support global clip plane for Planar Reflections",
+		ToolTip = "Whether to support the global clip plane needed for planar reflections.  Enabling this increases BasePass triangle cost by ~15% regardless of whether planar reflections are active. Changing this setting requires restarting the editor.",
+		ConfigRestartRequired = true))
+	uint32 bGlobalClipPlane : 1;
+
 	UPROPERTY(config, EditAnywhere, Category=VR, meta=(
 		ConsoleVariable="vr.InstancedStereo", DisplayName="Instanced Stereo",
 		ToolTip="Enable instanced stereo rendering (only available for D3D SM5 or PS4).",
diff --git a/Engine/Source/Runtime/Engine/Classes/Engine/Scene.h b/Engine/Source/Runtime/Engine/Classes/Engine/Scene.h
--- a/Engine/Source/Runtime/Engine/Classes/Engine/Scene.h
+++ b/Engine/Source/Runtime/Engine/Classes/Engine/Scene.h
@@ -871,7 +871,7 @@ struct FPostProcessSettings
 	float DepthOfFieldDepthBlurAmount;
 
 	/** CircleDOF only: Depth blur radius in pixels at 1920x */
-	UPROPERTY(interp, BlueprintReadWrite, AdvancedDisplay, Category=DepthOfField, meta=(ClampMin = "0.0", ClampMax = "4.0", editcondition = "bOverride_DepthOfFieldDepthBlurRadius", DisplayName = "Depth Blur Radius"))
+	UPROPERTY(interp, BlueprintReadWrite, AdvancedDisplay, Category=DepthOfField, meta=(ClampMin = "0.0", UIMax = "4.0", editcondition = "bOverride_DepthOfFieldDepthBlurRadius", DisplayName = "Depth Blur Radius"))
 	float DepthOfFieldDepthBlurRadius;
 
 	/** Artificial region where all content is in focus, starting after DepthOfFieldFocalDistance, in unreal units  (cm) */
diff --git a/Engine/Source/Runtime/Engine/Classes/Engine/SkeletalMesh.h b/Engine/Source/Runtime/Engine/Classes/Engine/SkeletalMesh.h
--- a/Engine/Source/Runtime/Engine/Classes/Engine/SkeletalMesh.h
+++ b/Engine/Source/Runtime/Engine/Classes/Engine/SkeletalMesh.h
@@ -547,13 +547,15 @@ struct FSkeletalMaterial
 	FSkeletalMaterial()
 		: MaterialInterface( NULL )
 		, bEnableShadowCasting( true )
+		, bRecomputeTangent( false )
 	{
 
 	}
 
-	FSkeletalMaterial( class UMaterialInterface* InMaterialInterface, bool bInEnableShadowCasting = true )
+	FSkeletalMaterial( class UMaterialInterface* InMaterialInterface, bool bInEnableShadowCasting = true, bool bInRecomputeTangent = false )
 		: MaterialInterface( InMaterialInterface )
 		, bEnableShadowCasting( bInEnableShadowCasting )
+		, bRecomputeTangent( bInRecomputeTangent )
 	{
 
 	}
@@ -568,6 +570,8 @@ struct FSkeletalMaterial
 	class UMaterialInterface *	MaterialInterface;
 	UPROPERTY(EditAnywhere, BlueprintReadOnly, Category=SkeletalMesh, Category=SkeletalMesh)
 	bool						bEnableShadowCasting;
+	UPROPERTY(EditAnywhere, BlueprintReadOnly, Category = SkeletalMesh, Category = SkeletalMesh)
+	bool						bRecomputeTangent;
 };
 
 class FSkeletalMeshResource;
diff --git a/Engine/Source/Runtime/Engine/Classes/Engine/SubsurfaceProfile.h b/Engine/Source/Runtime/Engine/Classes/Engine/SubsurfaceProfile.h
--- a/Engine/Source/Runtime/Engine/Classes/Engine/SubsurfaceProfile.h
+++ b/Engine/Source/Runtime/Engine/Classes/Engine/SubsurfaceProfile.h
@@ -158,10 +158,10 @@ private:
 
 // If you change this you need to recompile the SSS shaders.
 // Required if we use a texture format with limited size but want to express a larger radius
-#define SUBSURFACE_RADIUS_SCALE 1024.0f
+static const int32 SUBSURFACE_RADIUS_SCALE = 1024;
 
  // The kernels range from -3 to 3
-#define SUBSURFACE_KERNEL_SIZE 3.0f
+static const int32 SUBSURFACE_KERNEL_SIZE = 3;
 
 // lives on the render thread
 extern ENGINE_API TGlobalResource<FSubsurfaceProfileTexture> GSubsurfaceProfileTextureObject;
diff --git a/Engine/Source/Runtime/Engine/Classes/Engine/Texture2D.h b/Engine/Source/Runtime/Engine/Classes/Engine/Texture2D.h
--- a/Engine/Source/Runtime/Engine/Classes/Engine/Texture2D.h
+++ b/Engine/Source/Runtime/Engine/Classes/Engine/Texture2D.h
@@ -25,6 +25,18 @@ private:
 	int32 StreamingIndex;
 
 public:
+
+	/*
+	 * Level scope index of this texture. It is used to reduce the amount of lookup to map a texture to its level index.
+	 * Useful when building texture streaming data, as well as when filling the texture streamer with precomputed data.
+     * It relates to FStreamingTextureBuildInfo::TextureLevelIndex and also the index in ULevel::StreamingTextureGuids. 
+	 * Default value of -1, indicates that the texture has an unknown index (not yet processed). At level load time, 
+	 * -2 is also used to indicate that the texture has been processed but no entry were found in the level table.
+	 * After any of these processes, the LevelIndex is reset to INDEX_NONE. Making it ready for the next level task.
+	 */
+	UPROPERTY(transient, duplicatetransient, NonTransactional)
+	int32 LevelIndex;
+
 	/** keep track of first mip level used for ResourceMem creation */
 	UPROPERTY()
 	int32 FirstResourceMemMip;
