UE-67355 Editor Unable to Find Mali Offline Compiler On Mac

#MacOS
#4.22

FMacPlatformProcess::ExecProcess - fileExistsAtPath validation fails on Mac.
Addig parameters to the compile command line instead of the compiler path (the Mali executable).

diff --git a/Engine/Source/Developer/ShaderCompilerCommon/Private/ShaderCompilerCommon.cpp b/Engine/Source/Developer/ShaderCompilerCommon/Private/ShaderCompilerCommon.cpp
--- a/Engine/Source/Developer/ShaderCompilerCommon/Private/ShaderCompilerCommon.cpp
+++ b/Engine/Source/Developer/ShaderCompilerCommon/Private/ShaderCompilerCommon.cpp
@@ -647,6 +647,8 @@ void CompileOfflineMali(const FShaderCompilerInput& Input, FShaderCompilerOutput
 
 		FString CompilerPath = Input.ExtraSettings.OfflineCompilerPath;
 
+		FString CompilerCommand = "";
+
 		// add process and thread ids to the file name to avoid collision between workers
 		auto ProcID = FPlatformProcess::GetCurrentProcessId();
 		auto ThreadID = FPlatformTLS::GetCurrentThreadId();
@@ -658,27 +660,27 @@ void CompileOfflineMali(const FShaderCompilerInput& Input, FShaderCompilerOutput
 		{
 			case SF_Vertex:
 				GLSLSourceFile += TEXT(".vert");
-				CompilerPath += TEXT(" -v");
+				CompilerCommand += TEXT(" -v");
 			break;
 			case SF_Pixel:
 				GLSLSourceFile += TEXT(".frag");
-				CompilerPath += TEXT(" -f");
+				CompilerCommand += TEXT(" -f");
 			break;
 			case SF_Geometry:
 				GLSLSourceFile += TEXT(".geom");
-				CompilerPath += TEXT(" -g");
+				CompilerCommand += TEXT(" -g");
 			break;
 			case SF_Hull:
 				GLSLSourceFile += TEXT(".tesc");
-				CompilerPath += TEXT(" -t");
+				CompilerCommand += TEXT(" -t");
 			break;
 			case SF_Domain:
 				GLSLSourceFile += TEXT(".tese");
-				CompilerPath += TEXT(" -e");
+				CompilerCommand += TEXT(" -e");
 			break;
 			case SF_Compute:
 				GLSLSourceFile += TEXT(".comp");
-				CompilerPath += TEXT(" -C");
+				CompilerCommand += TEXT(" -C");
 			break;
 
 			default:
@@ -688,11 +690,11 @@ void CompileOfflineMali(const FShaderCompilerInput& Input, FShaderCompilerOutput
 
 		if (bVulkanSpirV)
 		{
-			CompilerPath += TEXT(" -p");
+			CompilerCommand += TEXT(" -p");
 		}
 		else
 		{
-			CompilerPath += TEXT(" -s");
+			CompilerCommand += TEXT(" -s");
 		}
 
 		FArchive* Ar = IFileManager::Get().CreateFileWriter(*GLSLSourceFile, FILEWRITE_EvenIfReadOnly);
@@ -710,8 +712,21 @@ void CompileOfflineMali(const FShaderCompilerInput& Input, FShaderCompilerOutput
 		FString StdErr;
 		int32 ReturnCode = 0;
 
-		// Run Mali shader compiler and wait for completion
-		FPlatformProcess::ExecProcess(*CompilerPath, *GLSLSourceFile, &ReturnCode, &StdOut, &StdErr);
+		// Since v6.2.0, Mali compiler needs to be started in the executable folder or it won't find "external/glslangValidator" for Vulkan
+		FString CompilerWorkingDirectory = FPaths::GetPath(CompilerPath);
+
+		if (!CompilerWorkingDirectory.IsEmpty() && FPaths::DirectoryExists(CompilerWorkingDirectory))
+		{
+			// compiler command line contains flags and the GLSL source file name
+			CompilerCommand += " " + FPaths::ConvertRelativePathToFull(GLSLSourceFile);
+
+			// Run Mali shader compiler and wait for completion
+			FPlatformProcess::ExecProcess(*CompilerPath, *CompilerCommand, &ReturnCode, &StdOut, &StdErr, *CompilerWorkingDirectory);
+		}
+		else
+		{
+			StdErr = "Couldn't find Mali offline compiler at " + CompilerPath;
+		}
 
 		// parse Mali's output and extract instruction count or eventual errors
 		ShaderOutput.bSucceeded = (ReturnCode >= 0);
diff --git a/Engine/Source/Runtime/Core/Private/GenericPlatform/GenericPlatformProcess.cpp b/Engine/Source/Runtime/Core/Private/GenericPlatform/GenericPlatformProcess.cpp
--- a/Engine/Source/Runtime/Core/Private/GenericPlatform/GenericPlatformProcess.cpp
+++ b/Engine/Source/Runtime/Core/Private/GenericPlatform/GenericPlatformProcess.cpp
@@ -273,7 +273,7 @@ FString FGenericPlatformProcess::GetApplicationName( uint32 ProcessId )
 	return FString(TEXT(""));
 }
 
-bool FGenericPlatformProcess::ExecProcess( const TCHAR* URL, const TCHAR* Params, int32* OutReturnCode, FString* OutStdOut, FString* OutStdErr )
+bool FGenericPlatformProcess::ExecProcess(const TCHAR* URL, const TCHAR* Params, int32* OutReturnCode, FString* OutStdOut, FString* OutStdErr, const TCHAR* OptionalWorkingDirectory)
 {
 	UE_LOG(LogHAL, Fatal, TEXT("FGenericPlatformProcess::ExecProcess not implemented on this platform"));
 	return false;
diff --git a/Engine/Source/Runtime/Core/Private/Mac/MacPlatformProcess.cpp b/Engine/Source/Runtime/Core/Private/Mac/MacPlatformProcess.cpp
--- a/Engine/Source/Runtime/Core/Private/Mac/MacPlatformProcess.cpp
+++ b/Engine/Source/Runtime/Core/Private/Mac/MacPlatformProcess.cpp
@@ -258,7 +258,7 @@ FString FMacPlatformProcess::GetGameBundleId()
 
 @end // NSAutoReadPipe
 
-bool FMacPlatformProcess::ExecProcess( const TCHAR* URL, const TCHAR* Params, int32* OutReturnCode, FString* OutStdOut, FString* OutStdErr )
+bool FMacPlatformProcess::ExecProcess( const TCHAR* URL, const TCHAR* Params, int32* OutReturnCode, FString* OutStdOut, FString* OutStdErr, const TCHAR* OptionalWorkingDirectory)
 {
 	SCOPED_AUTORELEASE_POOL;
 
@@ -301,6 +301,13 @@ bool FMacPlatformProcess::ExecProcess( const TCHAR* URL, const TCHAR* Params, in
 	if (ProcessHandle)
 	{
 		[ProcessHandle setLaunchPath: LaunchPath];
+
+		if (OptionalWorkingDirectory != NULL)
+		{
+			NSString* WorkingDirectory = (NSString*)FPlatformString::TCHARToCFString(OptionalWorkingDirectory);
+			[ProcessHandle setCurrentDirectoryPath : WorkingDirectory];
+			CFRelease((CFStringRef)WorkingDirectory);
+		}
 		
 		TArray<FString> ArgsArray;
 		FString(Params).ParseIntoArray(ArgsArray, TEXT(" "), true);
diff --git a/Engine/Source/Runtime/Core/Private/Unix/UnixPlatformProcess.cpp b/Engine/Source/Runtime/Core/Private/Unix/UnixPlatformProcess.cpp
--- a/Engine/Source/Runtime/Core/Private/Unix/UnixPlatformProcess.cpp
+++ b/Engine/Source/Runtime/Core/Private/Unix/UnixPlatformProcess.cpp
@@ -1472,7 +1472,7 @@ bool FUnixPlatformProcess::IsApplicationRunning( const TCHAR* ProcName )
 	return !system(TCHAR_TO_UTF8(*Commandline));
 }
 
-bool FUnixPlatformProcess::ExecProcess( const TCHAR* URL, const TCHAR* Params, int32* OutReturnCode, FString* OutStdOut, FString* OutStdErr )
+bool FUnixPlatformProcess::ExecProcess(const TCHAR* URL, const TCHAR* Params, int32* OutReturnCode, FString* OutStdOut, FString* OutStdErr, const TCHAR* OptionalWorkingDirectory)
 {
 	FString CmdLineParams = Params;
 	FString ExecutableFileName = URL;
@@ -1493,7 +1493,7 @@ bool FUnixPlatformProcess::ExecProcess( const TCHAR* URL, const TCHAR* Params, i
 	const bool bLaunchHidden = false;
 	const bool bLaunchReallyHidden = bLaunchHidden;
 
-	FProcHandle ProcHandle = FPlatformProcess::CreateProc(*ExecutableFileName, *CmdLineParams, bLaunchDetached, bLaunchHidden, bLaunchReallyHidden, NULL, 0, NULL, PipeWrite);
+	FProcHandle ProcHandle = FPlatformProcess::CreateProc(*ExecutableFileName, *CmdLineParams, bLaunchDetached, bLaunchHidden, bLaunchReallyHidden, NULL, 0, OptionalWorkingDirectory, PipeWrite);
 	if (ProcHandle.IsValid())
 	{
 		while (FPlatformProcess::IsProcRunning(ProcHandle))
diff --git a/Engine/Source/Runtime/Core/Private/Windows/WindowsPlatformProcess.cpp b/Engine/Source/Runtime/Core/Private/Windows/WindowsPlatformProcess.cpp
--- a/Engine/Source/Runtime/Core/Private/Windows/WindowsPlatformProcess.cpp
+++ b/Engine/Source/Runtime/Core/Private/Windows/WindowsPlatformProcess.cpp
@@ -582,7 +582,7 @@ void FWindowsPlatformProcess::ReadFromPipes(FString* OutStrings[], HANDLE InPipe
  * Executes a process, returning the return code, stdout, and stderr. This
  * call blocks until the process has returned.
  */
-bool FWindowsPlatformProcess::ExecProcess( const TCHAR* URL, const TCHAR* Params, int32* OutReturnCode, FString* OutStdOut, FString* OutStdErr )
+bool FWindowsPlatformProcess::ExecProcess(const TCHAR* URL, const TCHAR* Params, int32* OutReturnCode, FString* OutStdOut, FString* OutStdErr, const TCHAR* OptionalWorkingDirectory)
 {
 	STARTUPINFOEX StartupInfoEx;
 	ZeroMemory(&StartupInfoEx, sizeof(StartupInfoEx));
@@ -628,7 +628,7 @@ bool FWindowsPlatformProcess::ExecProcess( const TCHAR* URL, const TCHAR* Params
 	FString CommandLine = FString::Printf(TEXT("%s %s"), URL, Params);
 
 	PROCESS_INFORMATION ProcInfo;
-	if (CreateProcess(NULL, CommandLine.GetCharArray().GetData(), NULL, NULL, TRUE, NORMAL_PRIORITY_CLASS | DETACHED_PROCESS | EXTENDED_STARTUPINFO_PRESENT, NULL, NULL, &StartupInfoEx.StartupInfo, &ProcInfo))
+	if (CreateProcess(NULL, CommandLine.GetCharArray().GetData(), NULL, NULL, TRUE, NORMAL_PRIORITY_CLASS | DETACHED_PROCESS | EXTENDED_STARTUPINFO_PRESENT, NULL, OptionalWorkingDirectory, &StartupInfoEx.StartupInfo, &ProcInfo))
 	{
 		if (hStdOutRead != NULL)
 		{
diff --git a/Engine/Source/Runtime/Core/Public/GenericPlatform/GenericPlatformProcess.h b/Engine/Source/Runtime/Core/Public/GenericPlatform/GenericPlatformProcess.h
--- a/Engine/Source/Runtime/Core/Public/GenericPlatform/GenericPlatformProcess.h
+++ b/Engine/Source/Runtime/Core/Public/GenericPlatform/GenericPlatformProcess.h
@@ -216,7 +216,7 @@ struct CORE_API FGenericPlatformProcess
 	static const TCHAR* ShaderDir();
 	static void SetShaderDir(const TCHAR*Where);
 	static void SetCurrentWorkingDirectoryToBaseDir();
-	
+
 	/** Get the current working directory (only really makes sense on desktop platforms) */
 	static FString GetCurrentWorkingDirectory();
 
@@ -403,8 +403,9 @@ struct CORE_API FGenericPlatformProcess
 	 * @param OutReturnCode may be 0
 	 * @param OutStdOut may be 0
 	 * @param OutStdErr may be 0
+	 * @OptionalWorkingDirectory may be 0
 	 */
-	static bool ExecProcess( const TCHAR* URL, const TCHAR* Params, int32* OutReturnCode, FString* OutStdOut, FString* OutStdErr );
+	static bool ExecProcess(const TCHAR* URL, const TCHAR* Params, int32* OutReturnCode, FString* OutStdOut, FString* OutStdErr, const TCHAR* OptionalWorkingDirectory = NULL);
 
 	/**
 	 * Executes a process as administrator, requesting elevation as necessary. This
diff --git a/Engine/Source/Runtime/Core/Public/Mac/MacPlatformProcess.h b/Engine/Source/Runtime/Core/Public/Mac/MacPlatformProcess.h
--- a/Engine/Source/Runtime/Core/Public/Mac/MacPlatformProcess.h
+++ b/Engine/Source/Runtime/Core/Public/Mac/MacPlatformProcess.h
@@ -138,7 +138,7 @@ struct CORE_API FMacPlatformProcess : public FGenericPlatformProcess
 	static bool ReadPipeToArray(void* ReadPipe, TArray<uint8> & Output);
 	static bool WritePipe(void* WritePipe, const FString& Message, FString* OutWritten = nullptr);
 	static bool WritePipe(void* WritePipe, const uint8* Data, const int32 DataLength, int32* OutDataLength = nullptr);
-	static bool ExecProcess(const TCHAR* URL, const TCHAR* Params, int32* OutReturnCode, FString* OutStdOut, FString* OutStdErr);
+	static bool ExecProcess(const TCHAR* URL, const TCHAR* Params, int32* OutReturnCode, FString* OutStdOut, FString* OutStdErr, const TCHAR* OptionalWorkingDirectory = NULL);
 	static void SetThreadAffinityMask(uint64 AffinityMask);
 	
 	// Mac specific
diff --git a/Engine/Source/Runtime/Core/Public/Unix/UnixPlatformProcess.h b/Engine/Source/Runtime/Core/Public/Unix/UnixPlatformProcess.h
--- a/Engine/Source/Runtime/Core/Public/Unix/UnixPlatformProcess.h
+++ b/Engine/Source/Runtime/Core/Public/Unix/UnixPlatformProcess.h
@@ -273,7 +273,7 @@ struct CORE_API FUnixPlatformProcess : public FGenericPlatformProcess
 	static bool Daemonize();
 	static bool IsApplicationRunning( uint32 ProcessId );
 	static bool IsApplicationRunning( const TCHAR* ProcName );
-	static bool ExecProcess( const TCHAR* URL, const TCHAR* Params, int32* OutReturnCode, FString* OutStdOut, FString* OutStdErr );
+	static bool ExecProcess(const TCHAR* URL, const TCHAR* Params, int32* OutReturnCode, FString* OutStdOut, FString* OutStdErr, const TCHAR* OptionalWorkingDirectory = NULL);
 	static void ExploreFolder( const TCHAR* FilePath );
 	static void LaunchFileInDefaultExternalApplication( const TCHAR* FileName, const TCHAR* Parms = NULL, ELaunchVerb::Type Verb = ELaunchVerb::Open );
 	static bool IsFirstInstance();
diff --git a/Engine/Source/Runtime/Core/Public/Windows/WindowsPlatformProcess.h b/Engine/Source/Runtime/Core/Public/Windows/WindowsPlatformProcess.h
--- a/Engine/Source/Runtime/Core/Public/Windows/WindowsPlatformProcess.h
+++ b/Engine/Source/Runtime/Core/Public/Windows/WindowsPlatformProcess.h
@@ -158,7 +158,7 @@ public:
 	static bool IsApplicationRunning( uint32 ProcessId );
 	static bool IsApplicationRunning( const TCHAR* ProcName );
 	static FString GetApplicationName( uint32 ProcessId );	
-	static bool ExecProcess( const TCHAR* URL, const TCHAR* Params, int32* OutReturnCode, FString* OutStdOut, FString* OutStdErr );
+	static bool ExecProcess(const TCHAR* URL, const TCHAR* Params, int32* OutReturnCode, FString* OutStdOut, FString* OutStdErr, const TCHAR* OptionalWorkingDirectory = NULL);
 	static bool ExecElevatedProcess(const TCHAR* URL, const TCHAR* Params, int32* OutReturnCode);
 	static void LaunchFileInDefaultExternalApplication( const TCHAR* FileName, const TCHAR* Parms = NULL, ELaunchVerb::Type Verb = ELaunchVerb::Open );
 	static void ExploreFolder( const TCHAR* FilePath );
