Cherry-pick 4514566: Converts FDirectionalLightOcclusionRG to render graph and shader parameter struct API.

This notably implements SHADER_USE_ROOT_PARAMETER_STRUCT(), that is able to generate a cbuffer
for the shader's FParameters struct automatically, without requiring passing through a global shader parameter struct.

diff --git a/Engine/Plugins/FX/Niagara/Source/NiagaraShader/Public/NiagaraShaderType.h b/Engine/Plugins/FX/Niagara/Source/NiagaraShader/Public/NiagaraShaderType.h
--- a/Engine/Plugins/FX/Niagara/Source/NiagaraShader/Public/NiagaraShaderType.h
+++ b/Engine/Plugins/FX/Niagara/Source/NiagaraShader/Public/NiagaraShaderType.h
@@ -84,7 +84,7 @@ public:
 		ValidateCompiledResultType InValidateCompiledResultRef,
 		GetStreamOutElementsType InGetStreamOutElementsRef
 		):
-		FShaderType(EShaderTypeForDynamicCast::Niagara, InName, InSourceFilename, InFunctionName, SF_Compute, InTotalPermutationCount, InConstructSerializedRef, InGetStreamOutElementsRef),
+		FShaderType(EShaderTypeForDynamicCast::Niagara, InName, InSourceFilename, InFunctionName, SF_Compute, InTotalPermutationCount, InConstructSerializedRef, InGetStreamOutElementsRef, nullptr),
 		ConstructCompiledRef(InConstructCompiledRef),
 		ShouldCompilePermutationRef(InShouldCompilePermutationRef),
 		ValidateCompiledResultRef(InValidateCompiledResultRef),
diff --git a/Engine/Source/Developer/ShaderCompilerCommon/Private/ShaderCompilerCommon.cpp b/Engine/Source/Developer/ShaderCompilerCommon/Private/ShaderCompilerCommon.cpp
--- a/Engine/Source/Developer/ShaderCompilerCommon/Private/ShaderCompilerCommon.cpp
+++ b/Engine/Source/Developer/ShaderCompilerCommon/Private/ShaderCompilerCommon.cpp
@@ -432,6 +432,248 @@ TCHAR* FindNextUniformBufferReference(TCHAR* SearchPtr, const TCHAR* SearchStrin
 }
 
 
+void MoveShaderParametersToRootConstantBuffer(const FShaderCompilerInput& CompilerInput, FString& PreprocessedShaderSource)
+{
+	check(CompilerInput.RootParameterBindings.Num());
+
+	TMap<FString, FString> ShaderParameterTypes;
+	ShaderParameterTypes.Reserve(CompilerInput.RootParameterBindings.Num());
+
+	// Prepare the set of parameter to look for.
+	for (const auto& Member : CompilerInput.RootParameterBindings)
+	{
+		ShaderParameterTypes.Add(Member.Name, FString());
+	}
+
+
+	// Browse the code for global shader parameter, Save their type and erase them white spaces.
+	{
+		int32 TypeStartPos = -1;
+		int32 TypeEndPos = -1;
+		int32 NameStartPos = -1;
+		int32 NameEndPos = -1;
+		int32 ScopeIndent = 0;
+		bool bGoToNextSemicolon = false;
+		bool bGoToNextLine = false;
+
+		for (int32 Cursor = 0; Cursor < PreprocessedShaderSource.Len(); Cursor++)
+		{
+			TCHAR Char = PreprocessedShaderSource[Cursor];
+
+			const TCHAR* UpComing = (*PreprocessedShaderSource) + Cursor;
+
+			// Go to the next line if this is a preprocessor macro.
+			if (bGoToNextLine)
+			{
+				if (Char == '\n')
+				{
+					bGoToNextLine = false;
+				}
+				continue;
+			}
+			else if (Char == '#')
+			{
+				bGoToNextLine = true;
+				continue;
+			}
+
+			// If within a scope, just carry on until outside the scope.
+			if (ScopeIndent > 0 || Char == '{')
+			{
+				if (Char == '{')
+				{
+					ScopeIndent++;
+				}
+				else if (Char == '}')
+				{
+					ScopeIndent--;
+					if (ScopeIndent == 0)
+					{
+						TypeStartPos = -1;
+						TypeEndPos = -1;
+						NameStartPos = -1;
+						NameEndPos = -1;
+						bGoToNextSemicolon = false;
+					}
+				}
+				continue;
+			}
+
+			// If need to go to next global semicolon and reach it. Resume browsing.
+			if (bGoToNextSemicolon)
+			{
+				if (Char == ';')
+				{
+					bGoToNextSemicolon = false;
+				}
+				continue;
+			}
+
+			// Found something interesting...
+			if ((Char >= 'a' && Char <= 'z') ||
+				(Char >= 'A' && Char <= 'Z') ||
+				(Char >= '0' && Char <= '9') ||
+				Char == '<' || Char == '>')
+			{
+				if (TypeStartPos == -1)
+				{
+					TypeStartPos = Cursor;
+				}
+				else if (TypeEndPos != -1 && NameStartPos == -1)
+				{
+					NameStartPos = Cursor;
+				}
+				else if (NameEndPos != -1)
+				{
+					TypeStartPos = -1;
+					TypeEndPos = -1;
+					NameStartPos = -1;
+					NameEndPos = -1;
+					bGoToNextSemicolon = true;
+					continue;
+				}
+
+				continue;
+			}
+
+			// If this is white space, just carry on.
+			if (Char == ' ' || Char == '\t' || Char == '\r' || Char == '\n')
+			{
+				if (TypeStartPos != -1 && TypeEndPos == -1)
+				{
+					// Just finished browsing what might be a type.
+					TypeEndPos = Cursor - 1;
+				}
+				else if (NameStartPos != -1 && NameEndPos == -1)
+				{
+					// Just finished browsing what might be shader parameter name.
+					NameEndPos = Cursor - 1;
+				}
+				continue;
+			}
+			else if (Char == ';')
+			{
+				if (NameStartPos != -1 && NameEndPos == -1)
+				{
+					// Just finished browsing what is a shader parameter name.
+					NameEndPos = Cursor - 1;
+				}
+				else if (NameEndPos != -1)
+				{
+					// Greate we found something, so fall through.
+				}
+				else
+				{
+					// No idea what it was, reset...
+					TypeStartPos = -1;
+					TypeEndPos = -1;
+					NameStartPos = -1;
+					NameEndPos = -1;
+					continue;
+				}
+			}
+			else if (Char == ':' && NameStartPos != -1)
+			{
+				// Just finished browsing what might be shader parameter name.
+				if (NameEndPos != -1)
+				{
+					NameEndPos = Cursor - 1;
+				}
+				continue;
+			}
+			else
+			{
+				// No idea what it was, reset and go to next semicolon...
+				TypeStartPos = -1;
+				TypeEndPos = -1;
+				NameStartPos = -1;
+				NameEndPos = -1;
+				bGoToNextSemicolon = true;
+				continue;
+			}
+
+			check(Char == ';');
+
+			// A shader parameter has been found.
+			{
+				FString Type = PreprocessedShaderSource.Mid(TypeStartPos, TypeEndPos - TypeStartPos + 1);
+				FString Name = PreprocessedShaderSource.Mid(NameStartPos, NameEndPos - NameStartPos + 1);
+
+				if (ShaderParameterTypes.Contains(Name))
+				{
+					ensureMsgf(ShaderParameterTypes.FindChecked(Name).IsEmpty(), TEXT("Looks %s like shader parameter was duplicated."), *Name);
+					ShaderParameterTypes[Name] = Type;
+
+					// Erases this shader parameter conserving the same line numbers.
+					for (int32 j = TypeStartPos; j <= Cursor; j++)
+					{
+						if (PreprocessedShaderSource[j] != '\r' && PreprocessedShaderSource[j] != '\n')
+							PreprocessedShaderSource[j] = ' ';
+					}
+				}
+
+				// And reset.
+				TypeStartPos = -1;
+				TypeEndPos = -1;
+				NameStartPos = -1;
+				NameEndPos = -1;
+				bGoToNextSemicolon = false;
+			}
+		}
+	}
+
+	// Generate the root cbuffer content.
+	FString RootCBufferContent;
+	for (const auto& Member : CompilerInput.RootParameterBindings)
+	{
+		const FString& Type = ShaderParameterTypes[Member.Name];
+		if (Type.IsEmpty())
+		{
+			continue;
+		}
+
+		FString HLSLOffset;
+		{
+			int32 ByteOffset = int32(Member.ByteOffset);
+			HLSLOffset = FString::FromInt(ByteOffset / 16);
+			
+			switch (ByteOffset % 16)
+			{
+			case 0:
+				break;
+			case 4:
+				HLSLOffset.Append(TEXT(".y"));
+				break;
+			case 8:
+				HLSLOffset.Append(TEXT(".z"));
+				break;
+			case 12:
+				HLSLOffset.Append(TEXT(".w"));
+				break;
+			}
+		}
+
+		RootCBufferContent.Append(FString::Printf(
+			TEXT("%s %s : packoffset(c%s);\r\n"),
+			*Type,
+			*Member.Name,
+			*HLSLOffset));
+		ShaderParameterTypes.Add(Member.Name, FString());
+	}
+
+	FString NewShaderCode = FString::Printf(
+		TEXT("cbuffer %s\r\n")
+		TEXT("{\r\n")
+		TEXT("%s")
+		TEXT("}\r\n\r\n%s"),
+		FShaderParametersMetadata::kRootUniformBufferBindingName,
+		*RootCBufferContent,
+		*PreprocessedShaderSource);
+
+	PreprocessedShaderSource = MoveTemp(NewShaderCode);
+}
+
+
 // The cross compiler doesn't yet support struct initializers needed to construct static structs for uniform buffers
 // Replace all uniform buffer struct member references (View.WorldToClip) with a flattened name that removes the struct dependency (View_WorldToClip)
 void RemoveUniformBuffersFromSource(const FShaderCompilerEnvironment& Environment, FString& PreprocessedShaderSource)
@@ -535,7 +777,7 @@ FString CreateShaderCompilerWorkerDirectCommandLine(const FShaderCompilerInput&
 	case SF_Geometry:	Text += TEXT(" -gs"); break;
 	case SF_Pixel:		Text += TEXT(" -ps"); break;
 	case SF_Compute:	Text += TEXT(" -cs"); break;
-	default: ensure(0); break;
+	default: break;
 	}
 	if (Input.bCompilingForShaderPipeline)
 	{
diff --git a/Engine/Source/Developer/ShaderCompilerCommon/Public/ShaderCompilerCommon.h b/Engine/Source/Developer/ShaderCompilerCommon/Public/ShaderCompilerCommon.h
--- a/Engine/Source/Developer/ShaderCompilerCommon/Public/ShaderCompilerCommon.h
+++ b/Engine/Source/Developer/ShaderCompilerCommon/Public/ShaderCompilerCommon.h
@@ -40,6 +40,9 @@ extern SHADERCOMPILERCOMMON_API void BuildResourceTableTokenStream(
 // Finds the number of used uniform buffers in a resource map
 extern SHADERCOMPILERCOMMON_API int16 GetNumUniformBuffersUsed(const FShaderCompilerResourceTable& InSRT);
 
+/** Moves all the shader parameter defined in the root scope of the shader into the root uniform buffer. */
+extern SHADERCOMPILERCOMMON_API void MoveShaderParametersToRootConstantBuffer(const FShaderCompilerInput& CompilerInput, FString& PreprocessedShaderSource);
+
 // The cross compiler doesn't yet support struct initializers needed to construct static structs for uniform buffers
 // Replace all uniform buffer struct member references (View.WorldToClip) with a flattened name that removes the struct dependency (View_WorldToClip)
 extern SHADERCOMPILERCOMMON_API void RemoveUniformBuffersFromSource(const FShaderCompilerEnvironment& Environment, FString& PreprocessedShaderSource);
diff --git a/Engine/Source/Developer/Windows/ShaderFormatD3D/Private/D3D11ShaderCompiler.cpp b/Engine/Source/Developer/Windows/ShaderFormatD3D/Private/D3D11ShaderCompiler.cpp
--- a/Engine/Source/Developer/Windows/ShaderFormatD3D/Private/D3D11ShaderCompiler.cpp
+++ b/Engine/Source/Developer/Windows/ShaderFormatD3D/Private/D3D11ShaderCompiler.cpp
@@ -1073,6 +1073,10 @@ void CompileD3D11Shader(const FShaderCompilerInput& Input,FShaderCompilerOutput&
 		}
 	}
 
+	if (Input.RootParameterBindings.Num())
+	{
+		MoveShaderParametersToRootConstantBuffer(Input, PreprocessedShaderSource);
+	}
 	RemoveUniformBuffersFromSource(Input.Environment, PreprocessedShaderSource);
 
 	// Override default compiler path to newer dll
diff --git a/Engine/Source/Programs/ShaderCompileWorker/Private/ShaderCompileWorker.cpp b/Engine/Source/Programs/ShaderCompileWorker/Private/ShaderCompileWorker.cpp
--- a/Engine/Source/Programs/ShaderCompileWorker/Private/ShaderCompileWorker.cpp
+++ b/Engine/Source/Programs/ShaderCompileWorker/Private/ShaderCompileWorker.cpp
@@ -16,7 +16,7 @@
 #define DEBUG_USING_CONSOLE	0
 
 // this is for the protocol, not the data, bump if FShaderCompilerInput or ProcessInputFromArchive changes (also search for the second one with the same name, todo: put into one header file)
-const int32 ShaderCompileWorkerInputVersion = 9;
+const int32 ShaderCompileWorkerInputVersion = 10;
 // this is for the protocol, not the data, bump if FShaderCompilerOutput or WriteToOutputArchive changes (also search for the second one with the same name, todo: put into one header file)
 const int32 ShaderCompileWorkerOutputVersion = 5;
 // this is for the protocol, not the data, bump if FShaderCompilerOutput or WriteToOutputArchive changes (also search for the second one with the same name, todo: put into one header file)
diff --git a/Engine/Source/Runtime/Engine/Private/ShaderCompiler/ShaderCompiler.cpp b/Engine/Source/Runtime/Engine/Private/ShaderCompiler/ShaderCompiler.cpp
--- a/Engine/Source/Runtime/Engine/Private/ShaderCompiler/ShaderCompiler.cpp
+++ b/Engine/Source/Runtime/Engine/Private/ShaderCompiler/ShaderCompiler.cpp
@@ -46,6 +46,7 @@
 #include "SceneInterface.h"
 #include "ShaderCodeLibrary.h"
 #include "MeshMaterialShaderType.h"
+#include "ShaderParameterMetadata.h"
 
 #define LOCTEXT_NAMESPACE "ShaderCompiler"
 
@@ -80,7 +81,7 @@ FString GetMaterialShaderMapDDCKey()
 }
 
 // this is for the protocol, not the data, bump if FShaderCompilerInput or ProcessInputFromArchive changes (also search for the second one with the same name, todo: put into one header file)
-const int32 ShaderCompileWorkerInputVersion = 9;
+const int32 ShaderCompileWorkerInputVersion = 10;
 // this is for the protocol, not the data, bump if FShaderCompilerOutput or WriteToOutputArchive changes (also search for the second one with the same name, todo: put into one header file)
 const int32 ShaderCompileWorkerOutputVersion = 5;
 // this is for the protocol, not the data, bump if FShaderCompilerOutput or WriteToOutputArchive changes (also search for the second one with the same name, todo: put into one header file)
@@ -2735,6 +2736,51 @@ void ValidateShaderFilePath(const FString& VirtualShaderFilePath, const FString&
 		*VirtualShaderFilePath);
 }
 
+static void PullRootShaderParametersLayout(FShaderCompilerInput& CompileInput, const FShaderParametersMetadata& ParametersMetadata, uint16 ByteOffset)
+{
+	for (const FShaderParametersMetadata::FMember& Member : ParametersMetadata.GetMembers())
+	{
+		EUniformBufferBaseType BaseType = Member.GetBaseType();
+		uint16 MemberOffset = ByteOffset + uint16(Member.GetOffset());
+
+		if (BaseType == UBMT_NESTED_STRUCT || BaseType == UBMT_INCLUDED_STRUCT)
+		{
+			PullRootShaderParametersLayout(CompileInput, *Member.GetStructMetadata(), MemberOffset);
+		}
+		else if (
+			BaseType == UBMT_BOOL ||
+			BaseType == UBMT_INT32 ||
+			BaseType == UBMT_UINT32 ||
+			BaseType == UBMT_FLOAT32)
+		{
+			FShaderCompilerInput::FRootParameterBinding RootParameterBinding;
+			RootParameterBinding.Name = Member.GetName();
+			RootParameterBinding.ByteOffset = MemberOffset;
+			CompileInput.RootParameterBindings.Add(RootParameterBinding);
+		}
+		continue;
+
+		if (BaseType == UBMT_REFERENCED_STRUCT)
+		{
+			// Referenced structured are manually passed to the RHI.
+		}
+		else if (BaseType == UBMT_RENDER_TARGET_BINDING_SLOTS)
+		{
+			// RHI don't need to care about render target bindings slot anyway.
+		}
+		else if (
+			BaseType == UBMT_GRAPH_TRACKED_BUFFER_UAV ||
+			BaseType == UBMT_GRAPH_TRACKED_UAV)
+		{
+			// UAV are ignored on purpose because not supported in uniform buffers.
+		}
+		else
+		{
+			check(0);
+		}
+	}
+}
+
 TSharedPtr<FString> GCachedGeneratedInstancedStereoCode = MakeShareable(new FString());
 
 /** Enqueues a shader compile job with GShaderCompilingManager. */
@@ -2766,6 +2812,11 @@ void GlobalBeginCompileShader(
 	// asset material name or "Global"
 	Input.DebugGroupName = DebugGroupName;
 
+	if (ShaderType->GetRootParametersMetadata())
+	{
+		PullRootShaderParametersLayout(Input, *ShaderType->GetRootParametersMetadata(), /* ByteOffset = */ 0);
+	}
+
 	// Verify FShaderCompilerInput's file paths are consistent. 
 	#if DO_CHECK
 		check(CheckVirtualShaderFilePath(Input.VirtualSourceFilePath));
diff --git a/Engine/Source/Runtime/Engine/Public/MaterialShaderType.h b/Engine/Source/Runtime/Engine/Public/MaterialShaderType.h
--- a/Engine/Source/Runtime/Engine/Public/MaterialShaderType.h
+++ b/Engine/Source/Runtime/Engine/Public/MaterialShaderType.h
@@ -103,7 +103,7 @@ public:
 		ValidateCompiledResultType InValidateCompiledResultRef,
 		GetStreamOutElementsType InGetStreamOutElementsRef
 		):
-		FShaderType(EShaderTypeForDynamicCast::Material, InName, InSourceFilename, InFunctionName, InFrequency, InTotalPermutationCount, InConstructSerializedRef, InGetStreamOutElementsRef),
+		FShaderType(EShaderTypeForDynamicCast::Material, InName, InSourceFilename, InFunctionName, InFrequency, InTotalPermutationCount, InConstructSerializedRef, InGetStreamOutElementsRef, nullptr),
 		ConstructCompiledRef(InConstructCompiledRef),
 		ShouldCompilePermutationRef(InShouldCompilePermutationRef),
 		ValidateCompiledResultRef(InValidateCompiledResultRef),
diff --git a/Engine/Source/Runtime/Engine/Public/MeshMaterialShaderType.h b/Engine/Source/Runtime/Engine/Public/MeshMaterialShaderType.h
--- a/Engine/Source/Runtime/Engine/Public/MeshMaterialShaderType.h
+++ b/Engine/Source/Runtime/Engine/Public/MeshMaterialShaderType.h
@@ -76,7 +76,7 @@ public:
 		ValidateCompiledResultType InValidateCompiledResultRef,
 		GetStreamOutElementsType InGetStreamOutElementsRef
 		):
-		FShaderType(EShaderTypeForDynamicCast::MeshMaterial, InName,InSourceFilename,InFunctionName,InFrequency,InTotalPermutationCount,InConstructSerializedRef,InGetStreamOutElementsRef),
+		FShaderType(EShaderTypeForDynamicCast::MeshMaterial, InName,InSourceFilename,InFunctionName,InFrequency,InTotalPermutationCount,InConstructSerializedRef,InGetStreamOutElementsRef, nullptr),
 		ConstructCompiledRef(InConstructCompiledRef),
 		ShouldCompilePermutationRef(InShouldCompilePermutationRef),
 		ValidateCompiledResultRef(InValidateCompiledResultRef),
diff --git a/Engine/Source/Runtime/RHI/Public/RHIDefinitions.h b/Engine/Source/Runtime/RHI/Public/RHIDefinitions.h
--- a/Engine/Source/Runtime/RHI/Public/RHIDefinitions.h
+++ b/Engine/Source/Runtime/RHI/Public/RHIDefinitions.h
@@ -1124,6 +1124,22 @@ inline bool IsUniformBufferResourceType(EUniformBufferBaseType BaseType)
 		BaseType == UBMT_RENDER_TARGET_BINDING_SLOTS;
 }
 
+inline bool IsUniformBufferResourceTypeIgnoreByRHI(EUniformBufferBaseType BaseType)
+{
+	return BaseType == UBMT_RENDER_TARGET_BINDING_SLOTS;
+}
+
+inline bool IsUniformBufferResourceIndirectionType(EUniformBufferBaseType BaseType)
+{
+	return
+		BaseType == UBMT_GRAPH_TRACKED_TEXTURE ||
+		BaseType == UBMT_GRAPH_TRACKED_SRV ||
+		BaseType == UBMT_GRAPH_TRACKED_UAV ||
+		BaseType == UBMT_GRAPH_TRACKED_BUFFER ||
+		BaseType == UBMT_GRAPH_TRACKED_BUFFER_SRV ||
+		BaseType == UBMT_GRAPH_TRACKED_BUFFER_UAV;
+}
+
 inline const TCHAR* GetShaderFrequencyString(EShaderFrequency Frequency, bool bIncludePrefix = true)
 {
 	const TCHAR* String = TEXT("SF_NumFrequencies");
diff --git a/Engine/Source/Runtime/RenderCore/Private/Shader.cpp b/Engine/Source/Runtime/RenderCore/Private/Shader.cpp
--- a/Engine/Source/Runtime/RenderCore/Private/Shader.cpp
+++ b/Engine/Source/Runtime/RenderCore/Private/Shader.cpp
@@ -144,7 +144,8 @@ FShaderType::FShaderType(
 	uint32 InFrequency,
 	int32 InTotalPermutationCount,
 	ConstructSerializedType InConstructSerializedRef,
-	GetStreamOutElementsType InGetStreamOutElementsRef
+	GetStreamOutElementsType InGetStreamOutElementsRef,
+	const FShaderParametersMetadata* InRootParametersMetadata
 	):
 	ShaderTypeForDynamicCast(InShaderTypeForDynamicCast),
 	Name(InName),
@@ -155,6 +156,7 @@ FShaderType::FShaderType(
 	TotalPermutationCount(InTotalPermutationCount),
 	ConstructSerializedRef(InConstructSerializedRef),
 	GetStreamOutElementsRef(InGetStreamOutElementsRef),
+	RootParametersMetadata(InRootParametersMetadata),
 	GlobalListLink(this)
 {
 	bCachedUniformBufferStructDeclarations = false;
diff --git a/Engine/Source/Runtime/RenderCore/Private/ShaderParameterMetadata.cpp b/Engine/Source/Runtime/RenderCore/Private/ShaderParameterMetadata.cpp
--- a/Engine/Source/Runtime/RenderCore/Private/ShaderParameterMetadata.cpp
+++ b/Engine/Source/Runtime/RenderCore/Private/ShaderParameterMetadata.cpp
@@ -6,6 +6,7 @@
 
 #include "ShaderParameterMetadata.h"
 #include "RenderCore.h"
+#include "ShaderCore.h"
 
 
 static TLinkedList<FShaderParametersMetadata*>* GUniformStructList = nullptr;
@@ -146,11 +147,6 @@ void FShaderParametersMetadata::InitializeLayout()
 				{
 					UE_LOG(LogRendererCore, Fatal, TEXT("Shader parameter %s error: Shader parameter struct reference can only be done in shader parameter structs."), *CppName);
 				}
-
-				if (ChildStruct->GetUseCase() != FShaderParametersMetadata::EUseCase::GlobalShaderParameterStruct)
-				{
-					UE_LOG(LogRendererCore, Fatal, TEXT("Shader parameter %s error: Shader parameter struct can only reference global ones, but %s is not."), *CppName, ChildStruct->GetStructTypeName());
-				}
 			}
 
 			if (BaseType == UBMT_NESTED_STRUCT || BaseType == UBMT_INCLUDED_STRUCT)
@@ -160,7 +156,7 @@ void FShaderParametersMetadata::InitializeLayout()
 
 			if (UseCase == EUseCase::ShaderParameterStruct &&
 				(BaseType == UBMT_NESTED_STRUCT || BaseType == UBMT_INCLUDED_STRUCT) &&
-				ChildStruct->GetUseCase() != EUseCase::ShaderParameterStruct)
+				ChildStruct->GetUseCase() != EUseCase::ShaderParameterStruct && 0)
 			{
 				UE_LOG(LogRendererCore, Fatal, TEXT("Shader parameter %s error: shader parameter structs can only nests or include shader parameter struct, but %s is not."), *CppName, ChildStruct->GetStructTypeName());
 			}
diff --git a/Engine/Source/Runtime/RenderCore/Private/ShaderParameterStruct.cpp b/Engine/Source/Runtime/RenderCore/Private/ShaderParameterStruct.cpp
--- a/Engine/Source/Runtime/RenderCore/Private/ShaderParameterStruct.cpp
+++ b/Engine/Source/Runtime/RenderCore/Private/ShaderParameterStruct.cpp
@@ -22,6 +22,9 @@ struct FShaderParameterStructBindingContext
 	// C++ name of the render target binding slot.
 	FString RenderTargetBindingSlotCppName;
 
+	// Whether this is for legacy shader parameter settings, or root shader parameter structures/
+	bool bUseRootShaderParameters;
+
 
 	void Bind(
 		const FShaderParametersMetadata& StructMetaData,
@@ -72,6 +75,11 @@ struct FShaderParameterStructBindingContext
 				// The member name of a globally referenced struct is the not name on the struct.
 				ShaderBindingName = Member.GetStructMetadata()->GetShaderVariableName();
 			}
+			else if (bUseRootShaderParameters && (BaseType == UBMT_INT32 || BaseType == UBMT_UINT32 || BaseType == UBMT_FLOAT32))
+			{
+				// Constants are stored in the root shader parameter cbuffer when bUseRootShaderParameters == true.
+				continue;
+			}
 
 			if (ShaderGlobalScopeBindings.Contains(ShaderBindingName))
 			{
@@ -149,18 +157,21 @@ struct FShaderParameterStructBindingContext
 	}
 }; // struct FShaderParameterStructBindingContext
 
-void FShaderParameterBindings::Bind(const FShaderParameterMap& ParametersMap, const FShaderParametersMetadata& StructMetaData)
+void FShaderParameterBindings::BindForLegacyShaderParameters(const FShaderParameterMap& ParametersMap, const FShaderParametersMetadata& StructMetaData)
 {
 	checkf(StructMetaData.GetSize() < (1 << (sizeof(uint16) * 8)), TEXT("Shader parameter structure can only have a size < 65536 bytes."));
 
 	FShaderParameterStructBindingContext BindingContext;
 	BindingContext.Bindings = this;
 	BindingContext.ParametersMap = &ParametersMap;
+	BindingContext.bUseRootShaderParameters = false;
 	BindingContext.Bind(
 		StructMetaData,
 		/* MemberPrefix = */ TEXT(""),
 		/* ByteOffset = */ 0);
 
+	RootParameterBufferIndex = kInvalidBufferIndex;
+
 	TArray<FString> AllParameterNames;
 	ParametersMap.GetAllParameterNames(AllParameterNames);
 	if (0 && BindingContext.ShaderGlobalScopeBindings.Num() != AllParameterNames.Num()) // TODO: enable that once we can.
@@ -179,6 +190,48 @@ void FShaderParameterBindings::Bind(const FShaderParameterMap& ParametersMap, co
 	}
 }
 
+void FShaderParameterBindings::BindForRootShaderParameters(const FShaderParameterMap& ParametersMap, const FShaderParametersMetadata& StructMetaData)
+{
+	checkf(StructMetaData.GetSize() < (1 << (sizeof(uint16) * 8)), TEXT("Shader parameter structure can only have a size < 65536 bytes."));
+
+	FShaderParameterStructBindingContext BindingContext;
+	BindingContext.Bindings = this;
+	BindingContext.ParametersMap = &ParametersMap;
+	BindingContext.bUseRootShaderParameters = true;
+	BindingContext.Bind(
+		StructMetaData,
+		/* MemberPrefix = */ TEXT(""),
+		/* ByteOffset = */ 0);
+
+	// Binds the uniform buffer that contains the root shader parameters.
+	{
+		const TCHAR* ShaderBindingName = FShaderParametersMetadata::kRootUniformBufferBindingName;
+		uint16 BufferIndex, BaseIndex, BoundSize;
+		if (ParametersMap.FindParameterAllocation(ShaderBindingName, BufferIndex, BaseIndex, BoundSize))
+		{
+			BindingContext.ShaderGlobalScopeBindings.Add(ShaderBindingName, ShaderBindingName);
+			RootParameterBufferIndex = BufferIndex;
+		}
+	}
+
+	TArray<FString> AllParameterNames;
+	ParametersMap.GetAllParameterNames(AllParameterNames);
+	if (BindingContext.ShaderGlobalScopeBindings.Num() != AllParameterNames.Num())
+	{
+		UE_LOG(LogShaders, Error, TEXT("%i shader parameters have not been bound:"), AllParameterNames.Num() - BindingContext.ShaderGlobalScopeBindings.Num());
+		for (const FString& GlobalParameterName : AllParameterNames)
+		{
+			if (!BindingContext.ShaderGlobalScopeBindings.Contains(GlobalParameterName))
+			{
+				UE_LOG(LogShaders, Error, TEXT("  %s"), *GlobalParameterName);
+			}
+		}
+
+		// TODO: would be great to have the shader name for the error message.
+		UE_LOG(LogShaders, Fatal, TEXT("Some shader parameters have not been bound."));
+	}
+}
+
 void SetNullShaderParameterFatalError(const FShader* Shader, const FShaderParametersMetadata* ParametersMetadata, uint16 MemberOffset)
 {
 	const FShaderParametersMetadata* MemberContainingStruct = nullptr;
diff --git a/Engine/Source/Runtime/RenderCore/Public/GlobalShader.h b/Engine/Source/Runtime/RenderCore/Public/GlobalShader.h
--- a/Engine/Source/Runtime/RenderCore/Public/GlobalShader.h
+++ b/Engine/Source/Runtime/RenderCore/Public/GlobalShader.h
@@ -80,9 +80,10 @@ public:
 		ModifyCompilationEnvironmentType InModifyCompilationEnvironmentRef,
 		ShouldCompilePermutationType InShouldCompilePermutationRef,
 		ValidateCompiledResultType InValidateCompiledResultRef,
-		GetStreamOutElementsType InGetStreamOutElementsRef
+		GetStreamOutElementsType InGetStreamOutElementsRef,
+		const FShaderParametersMetadata* InRootParametersMetadata = nullptr
 		):
-		FShaderType(EShaderTypeForDynamicCast::Global, InName, InSourceFilename, InFunctionName, InFrequency, InTotalPermutationCount, InConstructSerializedRef, InGetStreamOutElementsRef),
+		FShaderType(EShaderTypeForDynamicCast::Global, InName, InSourceFilename, InFunctionName, InFrequency, InTotalPermutationCount, InConstructSerializedRef, InGetStreamOutElementsRef, InRootParametersMetadata),
 		ConstructCompiledRef(InConstructCompiledRef),
 		ShouldCompilePermutationRef(InShouldCompilePermutationRef),
 		ValidateCompiledResultRef(InValidateCompiledResultRef),
@@ -282,5 +283,6 @@ inline TShaderMap<FGlobalShaderType>* GetGlobalShaderMap(ERHIFeatureLevel::Type
 		ShaderClass::ModifyCompilationEnvironmentImpl, \
 		ShaderClass::ShouldCompilePermutationImpl, \
 		ShaderClass::ValidateCompiledResultImpl, \
-		ShaderClass::GetStreamOutElements \
+		ShaderClass::GetStreamOutElements, \
+		ShaderClass::GetRootParametersMetadata() \
 		)
diff --git a/Engine/Source/Runtime/RenderCore/Public/RenderGraphResources.h b/Engine/Source/Runtime/RenderCore/Public/RenderGraphResources.h
--- a/Engine/Source/Runtime/RenderCore/Public/RenderGraphResources.h
+++ b/Engine/Source/Runtime/RenderCore/Public/RenderGraphResources.h
@@ -18,11 +18,20 @@ class FRDGBufferSRV;
 class FRDGBufferUAV;
 
 struct FPooledRDGBuffer;
-
+struct FRaytracingShaderBindingsWriter;
 
 /** Generic graph resource. */
 class RENDERCORE_API FRDGResource
 {
+private:
+	// RHI resource once allocated.
+	union
+	{
+		mutable FRHIResource* Resource;
+		mutable FShaderResourceViewRHIParamRef SRV;
+		mutable FUnorderedAccessViewRHIParamRef UAV;
+	} CachedRHI;
+
 public:
 	// Name of the resource for debugging purpose.
 	const TCHAR* const Name = nullptr;
@@ -38,14 +47,6 @@ public:
 	void operator = (const FRDGResource&) = delete;
 
 private:
-	// RHI resource once allocated.
-	union
-	{
-		mutable FRHIResource* Resource;
-		mutable FShaderResourceViewRHIParamRef SRV;
-		mutable FUnorderedAccessViewRHIParamRef UAV;
-	} CachedRHI;
-
 	/** Number of references in passes and deferred queries. */
 	mutable int32 ReferenceCount = 0;
 
diff --git a/Engine/Source/Runtime/RenderCore/Public/Shader.h b/Engine/Source/Runtime/RenderCore/Public/Shader.h
--- a/Engine/Source/Runtime/RenderCore/Public/Shader.h
+++ b/Engine/Source/Runtime/RenderCore/Public/Shader.h
@@ -31,6 +31,7 @@ class FShaderPipelineType;
 class FShaderType;
 class FVertexFactoryParameterRef;
 class FVertexFactoryType;
+class FShaderParametersMetadata;
 
 
 /** Define a shader permutation uniquely according to its type, and permutation id.*/
@@ -585,6 +586,8 @@ class FMaterial;
 class RENDERCORE_API FShaderParameterBindings
 {
 public:
+	static constexpr uint16 kInvalidBufferIndex = 0xFFFF;
+
 	struct FParameter
 	{
 		uint16 BufferIndex;
@@ -633,6 +636,9 @@ public:
 	TArray<FResourceParameter> GraphUAVs;
 	TArray<FParameterStructReference> ParameterReferences;
 
+	// Buffer index of FShaderParametersMetadata::kRootUniformBufferBindingName
+	uint16 RootParameterBufferIndex;
+
 	friend FArchive& operator<<(FArchive& Ar, FShaderParameterBindings& ParametersBindingData)
 	{
 		Ar << ParametersBindingData.Parameters;
@@ -643,10 +649,12 @@ public:
 		Ar << ParametersBindingData.GraphSRVs;
 		Ar << ParametersBindingData.GraphUAVs;
 		Ar << ParametersBindingData.ParameterReferences;
+		Ar << ParametersBindingData.RootParameterBufferIndex;
 		return Ar;
 	}
 
-	void Bind(const FShaderParameterMap& ParameterMaps, const FShaderParametersMetadata& StructMetaData);
+	void BindForLegacyShaderParameters(const FShaderParameterMap& ParameterMaps, const FShaderParametersMetadata& StructMetaData);
+	void BindForRootShaderParameters(const FShaderParameterMap& ParameterMaps, const FShaderParametersMetadata& StructMetaData);
 }; // FShaderParameterBindings
 
 
@@ -894,6 +902,12 @@ public:
 	void DumpDebugInfo();
 	void SaveShaderStableKeys(EShaderPlatform TargetShaderPlatform, const struct FStableShaderKeyAndValue& SaveKeyVal);
 
+	/** Returns the meta data for the root shader parameter struct. */
+	static inline const FShaderParametersMetadata* GetRootParametersMetadata()
+	{
+		return nullptr;
+	}
+
 protected:
 
 	/** Indexed the same as UniformBufferParameters.  Packed densely for coherent traversal. */
@@ -996,7 +1010,8 @@ public:
 		uint32 InFrequency,
 		int32 TotalPermutationCount,
 		ConstructSerializedType InConstructSerializedRef,
-		GetStreamOutElementsType InGetStreamOutElementsRef);
+		GetStreamOutElementsType InGetStreamOutElementsRef,
+		const FShaderParametersMetadata* InRootParametersMetadata);
 
 	virtual ~FShaderType();
 
@@ -1096,6 +1111,12 @@ public:
 		return ReferencedUniformBufferStructsCache;
 	}
 
+	/** Returns the meta data for the root shader parameter struct. */
+	inline const FShaderParametersMetadata* GetRootParametersMetadata() const
+	{
+		return RootParametersMetadata;
+	}
+
 	/** Adds include statements for uniform buffers that this shader type references, and builds a prefix for the shader file with the include statements. */
 	void AddReferencedUniformBufferIncludes(FShaderCompilerEnvironment& OutEnvironment, FString& OutSourceFilePrefix, EShaderPlatform Platform);
 
@@ -1143,6 +1164,7 @@ private:
 
 	ConstructSerializedType ConstructSerializedRef;
 	GetStreamOutElementsType GetStreamOutElementsRef;
+	const FShaderParametersMetadata* const RootParametersMetadata;
 
 	/** A map from shader ID to shader.  A shader will be removed from it when deleted, so this doesn't need to use a TRefCountPtr. */
 	TMap<FShaderId,FShader*> ShaderIdMap;
diff --git a/Engine/Source/Runtime/RenderCore/Public/ShaderCore.h b/Engine/Source/Runtime/RenderCore/Public/ShaderCore.h
--- a/Engine/Source/Runtime/RenderCore/Public/ShaderCore.h
+++ b/Engine/Source/Runtime/RenderCore/Public/ShaderCore.h
@@ -526,6 +526,27 @@ struct FShaderCompilerInput
 	FShaderCompilerEnvironment Environment;
 	TRefCountPtr<FShaderCompilerEnvironment> SharedEnvironment;
 
+
+	struct FRootParameterBinding
+	{
+		/** Name of the constant buffer stored parameter. */
+		FString Name;
+
+		/** The offset of the parameter in the root shader parameter struct. */
+		uint16 ByteOffset;
+
+
+		friend FArchive& operator<<(FArchive& Ar, FRootParameterBinding& RootParameterBinding)
+		{
+			Ar << RootParameterBinding.Name;
+			Ar << RootParameterBinding.ByteOffset;
+			return Ar;
+		}
+	};
+
+	TArray<FRootParameterBinding> RootParameterBindings;
+
+
 	// Additional compilation settings that can be filled by FMaterial::SetupExtaCompilationSettings
 	// FMaterial::SetupExtaCompilationSettings is usually called by each (*)MaterialShaderType::BeginCompileShader() function
 	FExtraShaderCompilerSettings ExtraSettings;
@@ -644,6 +665,7 @@ struct FShaderCompilerInput
 		Ar << Input.DebugGroupName;
 		Ar << Input.Environment;
 		Ar << Input.ExtraSettings;
+		Ar << Input.RootParameterBindings;
 
 		// Note: skipping Input.SharedEnvironment, which is handled by FShaderCompileUtilities::DoWriteTasks in order to maintain sharing
 
diff --git a/Engine/Source/Runtime/RenderCore/Public/ShaderParameterMacros.h b/Engine/Source/Runtime/RenderCore/Public/ShaderParameterMacros.h
--- a/Engine/Source/Runtime/RenderCore/Public/ShaderParameterMacros.h
+++ b/Engine/Source/Runtime/RenderCore/Public/ShaderParameterMacros.h
@@ -418,6 +418,21 @@ struct TShaderParameterTypeInfo<FIntVector>
 	static const FShaderParametersMetadata* GetStructMetadata() { return nullptr; }
 };
 
+template<>
+struct TShaderParameterTypeInfo<FIntVector4>
+{
+	static constexpr EUniformBufferBaseType BaseType = UBMT_INT32;
+	static constexpr int32 NumRows = 1;
+	static constexpr int32 NumColumns = 4;
+	static constexpr int32 NumElements = 0;
+	static constexpr int32 Alignment = 16;
+	static constexpr bool bIsStoredInConstantBuffer = true;
+
+	using TAlignedType = TAlignedTypedef<FIntVector4, Alignment>::Type;
+
+	static const FShaderParametersMetadata* GetStructMetadata() { return nullptr; }
+};
+
 template<>
 struct TShaderParameterTypeInfo<FIntRect>
 {
diff --git a/Engine/Source/Runtime/RenderCore/Public/ShaderParameterMetadata.h b/Engine/Source/Runtime/RenderCore/Public/ShaderParameterMetadata.h
--- a/Engine/Source/Runtime/RenderCore/Public/ShaderParameterMetadata.h
+++ b/Engine/Source/Runtime/RenderCore/Public/ShaderParameterMetadata.h
@@ -49,6 +49,9 @@ public:
 		/** Shader parameter struct generated from assets, such as material parameter collection or Niagara. */
 		DataDrivenShaderParameterStruct,
 	};
+
+	/** Shader binding name of the uniform buffer that contains the root shader parameters. */
+	static constexpr const TCHAR* kRootUniformBufferBindingName = TEXT("_RootShaderParameters");
 	
 	/** A member of a shader parameter structure. */
 	class FMember
diff --git a/Engine/Source/Runtime/RenderCore/Public/ShaderParameterStruct.h b/Engine/Source/Runtime/RenderCore/Public/ShaderParameterStruct.h
--- a/Engine/Source/Runtime/RenderCore/Public/ShaderParameterStruct.h
+++ b/Engine/Source/Runtime/RenderCore/Public/ShaderParameterStruct.h
@@ -32,7 +32,19 @@
 	ShaderClass(const ShaderMetaType::CompiledShaderInitializerType& Initializer) \
 		: ShaderParentClass(Initializer) \
 	{ \
-		this->Bindings.Bind(Initializer.ParameterMap, *FParameters::FTypeInfo::GetStructMetadata()); \
+		this->Bindings.BindForLegacyShaderParameters(Initializer.ParameterMap, *FParameters::FTypeInfo::GetStructMetadata()); \
+	} \
+	\
+	ShaderClass() \
+	{ } \
+
+#define SHADER_USE_ROOT_PARAMETER_STRUCT(ShaderClass, ShaderParentClass) \
+	static inline const FShaderParametersMetadata* GetRootParametersMetadata() { return FParameters::FTypeInfo::GetStructMetadata(); } \
+	\
+	ShaderClass(const ShaderMetaType::CompiledShaderInitializerType& Initializer) \
+		: ShaderParentClass(Initializer) \
+	{ \
+		this->Bindings.BindForRootShaderParameters(Initializer.ParameterMap, *FParameters::FTypeInfo::GetStructMetadata()); \
 	} \
 	\
 	ShaderClass() \
@@ -86,6 +98,9 @@ inline void UnsetShaderUAVs(TRHICmdList& RHICmdList, const TShaderClass* Shader,
 {
 	// TODO(RDG): Once all shader sets their parameter through this, can refactor RHI so all UAVs of a shader get unset through a single RHI function call.
 	const FShaderParameterBindings& Bindings = Shader->Bindings;
+
+	checkf(Bindings.RootParameterBufferIndex == FShaderParameterBindings::kInvalidBufferIndex, TEXT("Can't use UnsetShaderUAVs() for root parameter buffer index."));
+
 	for (const FShaderParameterBindings::FResourceParameter& ParameterBinding : Bindings.GraphUAVs)
 	{
 		RHICmdList.SetUAVParameter(ShadeRHI, ParameterBinding.BaseIndex, nullptr);
@@ -103,6 +118,8 @@ inline void SetShaderParameters(TRHICmdList& RHICmdList, const TShaderClass* Sha
 	const typename TShaderClass::FParameters* ParametersPtr = &Parameters;
 	const uint8* Base = reinterpret_cast<const uint8*>(ParametersPtr);
 
+	checkf(Bindings.RootParameterBufferIndex == FShaderParameterBindings::kInvalidBufferIndex, TEXT("Can't use SetShaderParameters() for root parameter buffer index."));
+
 	// Parameters
 	for (const FShaderParameterBindings::FParameter& ParameterBinding : Bindings.Parameters)
 	{
