Extensions and bug fixes in FHlslParser.

#rnx

diff --git a/Engine/Shaders/Public/ShaderVersion.ush b/Engine/Shaders/Public/ShaderVersion.ush
--- a/Engine/Shaders/Public/ShaderVersion.ush
+++ b/Engine/Shaders/Public/ShaderVersion.ush
@@ -2,4 +2,6 @@
 // Each time the console command is executed it generates a new GUID. As this file is included
 // in Platform.ush (which should be included in any shader) it allows to invalidate the shader DDC.
 // 
-// GUID = 9686A7C745F60675F92F2DBFF3598712
\ No newline at end of file
+// If you are merging streams and there is a conflict with this GUID you should make a new GUID rather than taking one or the other.
+// GUID = 841DBB52-59CA-447E-A326-8A60200F54D6
+
diff --git a/Engine/Source/Developer/ShaderCompilerCommon/Private/HlslAST.cpp b/Engine/Source/Developer/ShaderCompilerCommon/Private/HlslAST.cpp
--- a/Engine/Source/Developer/ShaderCompilerCommon/Private/HlslAST.cpp
+++ b/Engine/Source/Developer/ShaderCompilerCommon/Private/HlslAST.cpp
@@ -249,7 +249,7 @@ namespace CrossCompiler
 				break;
 
 			case EOperators::FloatConstant:
-				Writer << FloatConstant;
+				Writer << Identifier;
 				break;
 
 			case EOperators::InitializerList:
@@ -262,6 +262,10 @@ namespace CrossCompiler
 			case EOperators::ArrayIndex:
 				break;
 
+			case EOperators::Sequence:
+				Writer << TEXT(",");
+				break;
+
 			case EOperators::TypeCast:
 				Writer << (TCHAR)'(';
 				TypeSpecifier->Write(Writer);
@@ -388,12 +392,19 @@ namespace CrossCompiler
 			switch (Operator)
 			{
 			case EOperators::ArrayIndex:
-				Writer << (TCHAR)'(';
-				SubExpressions[0]->Write(Writer);
-				Writer << TEXT(")");
+				if (SubExpressions[0]->AsUnaryExpression() && SubExpressions[0]->Operator == EOperators::Identifier)
+				{
+					SubExpressions[0]->Write(Writer);
+				}
+				else
+				{
+					Writer << (TCHAR)'(';
+					SubExpressions[0]->Write(Writer);
+					Writer << (TCHAR)')';
+				}
 				Writer << (TCHAR)'[';
 				SubExpressions[1]->Write(Writer);
-				Writer << TEXT("]");
+				Writer << (TCHAR)']';
 				break;
 
 			default:
@@ -851,12 +862,41 @@ namespace CrossCompiler
 			delete Specifier;
 		}
 
+		FRegisterSpecifier::FRegisterSpecifier(FLinearAllocator* InAllocator, const FSourceInfo& InInfo) :
+			FNode(InAllocator, InInfo),
+			Arguments(InAllocator)
+		{
+		}
+
+		FRegisterSpecifier::~FRegisterSpecifier()
+		{
+			for (FExpression* Expr : Arguments)
+			{
+				delete Expr;
+			}
+		}
+
+		void FRegisterSpecifier::Write(FASTWriter& Writer) const
+		{
+			Writer << TEXT("register(");
+			for (int32 Index = 0, Num = Arguments.Num(); Index < Num; ++Index)
+			{
+				Arguments[Index]->Write(Writer);
+				if (Index + 1 < Num)
+				{
+					Writer << TEXT(", ");
+				}
+			}
+			Writer << TEXT(")");
+		}
+
 		FDeclaration::FDeclaration(FLinearAllocator* InAllocator, const FSourceInfo& InInfo) :
 			FNode(InAllocator, InInfo),
 			Identifier(nullptr),
 			Semantic(nullptr),
 			bIsArray(false),
 			ArraySize(InAllocator),
+			Register(nullptr),
 			Initializer(nullptr)
 		{
 		}
@@ -868,6 +908,12 @@ namespace CrossCompiler
 
 			WriteOptionArraySize(Writer, bIsArray, ArraySize);
 
+			if (Register)
+			{
+				Writer << TEXT(" : ");
+				Register->Write(Writer);
+			}
+
 			if (Initializer)
 			{
 				Writer << TEXT(" = ");
diff --git a/Engine/Source/Developer/ShaderCompilerCommon/Private/HlslAST.h b/Engine/Source/Developer/ShaderCompilerCommon/Private/HlslAST.h
--- a/Engine/Source/Developer/ShaderCompilerCommon/Private/HlslAST.h
+++ b/Engine/Source/Developer/ShaderCompilerCommon/Private/HlslAST.h
@@ -244,7 +244,7 @@ namespace CrossCompiler
 			FloatConstant,
 			BoolConstant,
 
-			//Sequence,
+			Sequence,
 
 			TypeCast,
 		};
@@ -343,6 +343,9 @@ namespace CrossCompiler
 			case EHlslToken::Mod:
 				return AST::EOperators::Mod;
 
+			case EHlslToken::Comma:
+				return AST::EOperators::Sequence;
+
 			default:
 				check(0);
 				break;
@@ -371,7 +374,6 @@ namespace CrossCompiler
 			union
 			{
 				uint32 UintConstant;
-				float FloatConstant;
 				bool BoolConstant;
 				struct FTypeSpecifier* TypeSpecifier;
 			};
@@ -414,10 +416,10 @@ namespace CrossCompiler
 					return UintConstant;
 
 				case EOperators::FloatConstant:
-					return (uint32)FloatConstant;
+					return static_cast<uint32>(FCString::Atof(Identifier));
 
 				case EOperators::BoolConstant:
-					return (uint32)BoolConstant;
+					return static_cast<uint32>(BoolConstant);
 
 				default:
 					break;
@@ -460,6 +462,16 @@ namespace CrossCompiler
 			virtual FCompoundStatement* AsCompoundStatement() override { return this; }
 		};
 
+		struct FRegisterSpecifier : public FNode
+		{
+			FRegisterSpecifier(FLinearAllocator* InAllocator, const FSourceInfo& InInfo);
+			~FRegisterSpecifier();
+
+			virtual void Write(FASTWriter& Writer) const override;
+
+			TLinearArray<FExpression*> Arguments;
+		};
+
 		struct FDeclaration : public FNode
 		{
 			FDeclaration(FLinearAllocator* InAllocator, const FSourceInfo& InInfo);
@@ -477,6 +489,8 @@ namespace CrossCompiler
 
 			TLinearArray<FExpression*> ArraySize;
 
+			FRegisterSpecifier* Register;
+
 			FExpression* Initializer;
 		};
 
diff --git a/Engine/Source/Developer/ShaderCompilerCommon/Private/HlslExpressionParser.inl b/Engine/Source/Developer/ShaderCompilerCommon/Private/HlslExpressionParser.inl
--- a/Engine/Source/Developer/ShaderCompilerCommon/Private/HlslExpressionParser.inl
+++ b/Engine/Source/Developer/ShaderCompilerCommon/Private/HlslExpressionParser.inl
@@ -21,7 +21,7 @@ namespace CrossCompiler
 	struct FSymbolScope;
 	//struct FInfo;
 
-	EParseResult ComputeExpr(FHlslScanner& Scanner, int32 MinPrec, /*FInfo& Info,*/ FSymbolScope* SymbolScope, bool bAllowAssignment, FLinearAllocator* Allocator, AST::FExpression** OutExpression, AST::FExpression** OutTernaryExpression);
+	EParseResult ComputeExpr(FHlslScanner& Scanner, int32 MinPrec, /*FInfo& Info,*/ FSymbolScope* SymbolScope, int32 ExpressionFlags, FLinearAllocator* Allocator, AST::FExpression** OutExpression, AST::FExpression** OutTernaryExpression);
 	EParseResult ParseExpressionList(EHlslToken EndListToken, FHlslScanner& Scanner, FSymbolScope* SymbolScope, EHlslToken NewStartListToken, FLinearAllocator* Allocator, AST::FExpression* OutExpression);
 
 	struct FSymbolScope
@@ -172,6 +172,12 @@ namespace CrossCompiler
 		ETF_ERROR_IF_NOT_USER_TYPE	= 1 << 4,
 	};
 
+	enum EExpressionFlags
+	{
+		EEF_ALLOW_ASSIGNMENT	= 1 << 0,
+		EEF_ALLOW_SEQUENCE		= 1 << 1,
+	};
+
 	EParseResult ParseGeneralType(const FHlslToken* Token, int32 TypeFlags, FLinearAllocator* Allocator, AST::FTypeSpecifier** OutSpecifier)
 	{
 		if (!Token)
@@ -327,6 +333,7 @@ namespace CrossCompiler
 		case EHlslToken::RWTexture2DArray:
 		case EHlslToken::RWTexture3D:
 		case EHlslToken::StructuredBuffer:
+		case EHlslToken::RaytracingAccelerationStructure:
 			if (TypeFlags & ETF_SAMPLER_TEXTURE_BUFFER)
 			{
 				bMatched = true;
@@ -405,7 +412,7 @@ namespace CrossCompiler
 			auto* Token1 = Scanner.PeekToken(1);
 			auto* Token2 = Scanner.PeekToken(2);
 			FString TypeString;
-			if (Token && Token->Token == EHlslToken::Identifier && Token1 && Token1->Token == EHlslToken::ColonColon && Token2 && Token2->Token == EHlslToken::Identifier)
+			if (Token && Token->Token == EHlslToken::Identifier && Token1 && Token1->Token == EHlslToken::ColonColon && Token2 && Token2->Token == EHlslToken::Identifier && SymbolScope)
 			{
 				auto* Namespace = SymbolScope->GetGlobalScope();
 				do
@@ -600,7 +607,7 @@ namespace CrossCompiler
 		return ParseResultError();
 	}
 
-	EParseResult MatchSuffixOperator(FHlslScanner& Scanner, /*FInfo& Info,*/ FSymbolScope* SymbolScope, bool bAllowAssignment, FLinearAllocator* Allocator, AST::FExpression** InOutExpression, AST::FExpression** OutTernaryExpression)
+	EParseResult MatchSuffixOperator(FHlslScanner& Scanner, /*FInfo& Info,*/ FSymbolScope* SymbolScope, int32 ExpressionFlags, FLinearAllocator* Allocator, AST::FExpression** InOutExpression, AST::FExpression** OutTernaryExpression)
 	{
 		bool bFoundAny = false;
 		bool bTryAgain = true;
@@ -616,7 +623,7 @@ namespace CrossCompiler
 			{
 				Scanner.Advance();
 				AST::FExpression* ArrayIndex = nullptr;
-				auto Result = ComputeExpr(Scanner, 1, /*Info,*/ SymbolScope, bAllowAssignment, Allocator, &ArrayIndex, nullptr);
+				auto Result = ComputeExpr(Scanner, 1, /*Info,*/ SymbolScope, ExpressionFlags, Allocator, &ArrayIndex, nullptr);
 				if (Result != EParseResult::Matched)
 				{
 					Scanner.SourceError(TEXT("Expected expression!"));
@@ -687,7 +694,7 @@ namespace CrossCompiler
 			{
 				Scanner.Advance();
 				AST::FExpression* Left = nullptr;
-				if (ComputeExpr(Scanner, 0, /*Info,*/ SymbolScope, true, Allocator, &Left, nullptr) != EParseResult::Matched)
+				if (ComputeExpr(Scanner, 0, /*Info,*/ SymbolScope, (ExpressionFlags | EEF_ALLOW_ASSIGNMENT), Allocator, &Left, nullptr) != EParseResult::Matched)
 				{
 					Scanner.SourceError(TEXT("Expected expression!"));
 					return ParseResultError();
@@ -698,7 +705,7 @@ namespace CrossCompiler
 					return ParseResultError();
 				}
 				AST::FExpression* Right = nullptr;
-				if (ComputeExpr(Scanner, 0, /*Info,*/ SymbolScope, true, Allocator, &Right, nullptr) != EParseResult::Matched)
+				if (ComputeExpr(Scanner, 0, /*Info,*/ SymbolScope, (ExpressionFlags | EEF_ALLOW_ASSIGNMENT), Allocator, &Right, nullptr) != EParseResult::Matched)
 				{
 					Scanner.SourceError(TEXT("Expected expression!"));
 					return ParseResultError();
@@ -720,7 +727,7 @@ namespace CrossCompiler
 		return bFoundAny ? EParseResult::Matched : EParseResult::NotMatched;
 	}
 
-	EParseResult ComputeAtom(FHlslScanner& Scanner, /*FInfo& Info,*/ FSymbolScope* SymbolScope, bool bAllowAssignment, FLinearAllocator* Allocator, AST::FExpression** OutExpression, AST::FExpression** OutTernaryExpression)
+	EParseResult ComputeAtom(FHlslScanner& Scanner, /*FInfo& Info,*/ FSymbolScope* SymbolScope, int32 ExpressionFlags, FLinearAllocator* Allocator, AST::FExpression** OutExpression, AST::FExpression** OutTernaryExpression)
 	{
 		AST::FExpression* InnerUnaryExpression = nullptr;
 		auto UnaryResult = MatchUnaryOperator(Scanner, /*Info,*/ SymbolScope, Allocator, OutExpression, &InnerUnaryExpression);
@@ -748,7 +755,7 @@ namespace CrossCompiler
 		case EHlslToken::FloatConstant:
 			Scanner.Advance();
 			AtomExpression = new(Allocator) AST::FUnaryExpression(Allocator, AST::EOperators::FloatConstant, nullptr, Token->SourceInfo);
-			AtomExpression->FloatConstant = Token->Float;
+			AtomExpression->Identifier = Allocator->Strdup(Token->String);
 			break;
 
 		case EHlslToken::Identifier:
@@ -789,7 +796,7 @@ namespace CrossCompiler
 			const auto* Peek1 = Scanner.PeekToken(0);
 			const auto* Peek2 = Scanner.PeekToken(1);
 			// Parenthesis expression
-			if (ComputeExpr(Scanner, 1, /*Info,*/ SymbolScope, bAllowAssignment, Allocator, &AtomExpression, nullptr) != EParseResult::Matched)
+			if (ComputeExpr(Scanner, 1, /*Info,*/ SymbolScope, (ExpressionFlags | EEF_ALLOW_SEQUENCE), Allocator, &AtomExpression, nullptr) != EParseResult::Matched)
 			{
 				Scanner.SourceError(TEXT("Expected expression!"));
 				return ParseResultError();
@@ -854,7 +861,7 @@ namespace CrossCompiler
 		}
 
 		check(AtomExpression);
-		auto SuffixResult = MatchSuffixOperator(Scanner, /*Info,*/ SymbolScope, bAllowAssignment, Allocator, &AtomExpression, OutTernaryExpression);
+		auto SuffixResult = MatchSuffixOperator(Scanner, /*Info,*/ SymbolScope, ExpressionFlags, Allocator, &AtomExpression, OutTernaryExpression);
 		//auto* Token = Scanner.GetCurrentToken();
 		if (/*!Token || */SuffixResult == EParseResult::Error)
 		{
@@ -882,6 +889,9 @@ namespace CrossCompiler
 		{
 			switch (Token->Token)
 			{
+			case EHlslToken::Comma:
+				return 1;
+
 			case EHlslToken::Equal:
 			case EHlslToken::PlusEqual:
 			case EHlslToken::MinusEqual:
@@ -893,49 +903,49 @@ namespace CrossCompiler
 			case EHlslToken::AndEqual:
 			case EHlslToken::OrEqual:
 			case EHlslToken::XorEqual:
-				return 1;
+				return 2;
 
 			case EHlslToken::Question:
 				check(0);
-				return 2;
+				return 3;
 
 			case EHlslToken::OrOr:
-				return 3;
+				return 4;
 
 			case EHlslToken::AndAnd:
-				return 4;
+				return 5;
 
 			case EHlslToken::Or:
-				return 5;
+				return 6;
 
 			case EHlslToken::Xor:
-				return 6;
+				return 7;
 
 			case EHlslToken::And:
-				return 7;
+				return 8;
 
 			case EHlslToken::EqualEqual:
 			case EHlslToken::NotEqual:
-				return 8;
+				return 9;
 
 			case EHlslToken::Lower:
 			case EHlslToken::Greater:
 			case EHlslToken::LowerEqual:
 			case EHlslToken::GreaterEqual:
-				return 9;
+				return 10;
 
 			case EHlslToken::LowerLower:
 			case EHlslToken::GreaterGreater:
-				return 10;
+				return 11;
 
 			case EHlslToken::Plus:
 			case EHlslToken::Minus:
-				return 11;
+				return 12;
 
 			case EHlslToken::Times:
 			case EHlslToken::Div:
 			case EHlslToken::Mod:
-				return 12;
+				return 13;
 
 			default:
 				break;
@@ -982,6 +992,15 @@ namespace CrossCompiler
 		return false;
 	}
 
+	bool IsSequenceOperator(const FHlslToken* Token)
+	{
+		if (Token)
+		{
+			return (Token->Token == EHlslToken::Comma);
+		}
+		return false;
+	}
+
 	bool IsRightAssociative(const FHlslToken* Token)
 	{
 		return IsTernaryOperator(Token);
@@ -989,8 +1008,11 @@ namespace CrossCompiler
 
 	// Ternary is handled by popping out so it can right associate
 	//#todo-rco: Fix the case for right-associative assignment operators
-	EParseResult ComputeExpr(FHlslScanner& Scanner, int32 MinPrec, /*FInfo& Info,*/ FSymbolScope* SymbolScope, bool bAllowAssignment, FLinearAllocator* Allocator, AST::FExpression** OutExpression, AST::FExpression** OutTernaryExpression)
+	EParseResult ComputeExpr(FHlslScanner& Scanner, int32 MinPrec, /*FInfo& Info,*/ FSymbolScope* SymbolScope, int32 ExpressionFlags, FLinearAllocator* Allocator, AST::FExpression** OutExpression, AST::FExpression** OutTernaryExpression)
 	{
+		const bool bAllowAssignment = ((ExpressionFlags & EEF_ALLOW_ASSIGNMENT) == EEF_ALLOW_ASSIGNMENT);
+		const bool bAllowSequence = ((ExpressionFlags & EEF_ALLOW_SEQUENCE) == EEF_ALLOW_SEQUENCE);
+
 		auto OriginalToken = Scanner.GetCurrentTokenIndex();
 		//FInfoIndentScope Scope(Info);
 		/*
@@ -1012,7 +1034,7 @@ namespace CrossCompiler
 		*/
 		//Info.PrintWithTabs(FString::Printf(TEXT("Compute Expr %d\n"), MinPrec));
 		AST::FExpression* TernaryExpression = nullptr;
-		auto Result = ComputeAtom(Scanner, /*Info,*/ SymbolScope, bAllowAssignment, Allocator, OutExpression, &TernaryExpression);
+		auto Result = ComputeAtom(Scanner, /*Info,*/ SymbolScope, ExpressionFlags, Allocator, OutExpression, &TernaryExpression);
 		if (Result != EParseResult::Matched)
 		{
 			return Result;
@@ -1022,7 +1044,7 @@ namespace CrossCompiler
 		{
 			auto* Token = Scanner.GetCurrentToken();
 			int32 Precedence = GetPrecedence(Token);
-			if (!Token || !IsBinaryOperator(Token) || Precedence < MinPrec || (!bAllowAssignment && IsAssignmentOperator(Token)) || (OutTernaryExpression && *OutTernaryExpression))
+			if (!Token || !IsBinaryOperator(Token) || Precedence < MinPrec || (!bAllowAssignment && IsAssignmentOperator(Token)) || (!bAllowSequence && IsSequenceOperator(Token)) || (OutTernaryExpression && *OutTernaryExpression))
 			{
 				break;
 			}
@@ -1031,7 +1053,7 @@ namespace CrossCompiler
 			auto NextMinPrec = IsRightAssociative(Token) ? Precedence : Precedence + 1;
 			AST::FExpression* RHSExpression = nullptr;
 			AST::FExpression* RHSTernaryExpression = nullptr;
-			Result = ComputeExpr(Scanner, NextMinPrec, /*Info,*/ SymbolScope, bAllowAssignment, Allocator, &RHSExpression, &RHSTernaryExpression);
+			Result = ComputeExpr(Scanner, NextMinPrec, /*Info,*/ SymbolScope, ExpressionFlags, Allocator, &RHSExpression, &RHSTernaryExpression);
 			if (Result == EParseResult::Error)
 			{
 				return ParseResultError();
@@ -1080,10 +1102,10 @@ namespace CrossCompiler
 		return EParseResult::Matched;
 	}
 
-	EParseResult ParseExpression(FHlslScanner& Scanner, FSymbolScope* SymbolScope, bool bAllowAssignment, FLinearAllocator* Allocator, AST::FExpression** OutExpression)
+	EParseResult ParseExpression(FHlslScanner& Scanner, FSymbolScope* SymbolScope, int32 ExpressionFlags, FLinearAllocator* Allocator, AST::FExpression** OutExpression)
 	{
 		/*FInfo Info(!true);*/
-		return ComputeExpr(Scanner, 0, /*Info,*/ SymbolScope, bAllowAssignment, Allocator, OutExpression, nullptr);
+		return ComputeExpr(Scanner, 0, /*Info,*/ SymbolScope, ExpressionFlags, Allocator, OutExpression, nullptr);
 	}
 
 	EParseResult ParseExpressionList(EHlslToken EndListToken, FHlslScanner& Scanner, FSymbolScope* SymbolScope, EHlslToken NewStartListToken, FLinearAllocator* Allocator, AST::FExpression* OutExpression)
@@ -1112,7 +1134,7 @@ namespace CrossCompiler
 			else
 			{
 				AST::FExpression* Expression = nullptr;
-				auto Result = ParseExpression(Scanner, SymbolScope, true, Allocator, &Expression);
+				auto Result = ParseExpression(Scanner, SymbolScope, EEF_ALLOW_ASSIGNMENT, Allocator, &Expression);
 				if (Result == EParseResult::Error)
 				{
 					Scanner.SourceError(TEXT("Invalid expression list\n"));
diff --git a/Engine/Source/Developer/ShaderCompilerCommon/Private/HlslLexer.cpp b/Engine/Source/Developer/ShaderCompilerCommon/Private/HlslLexer.cpp
--- a/Engine/Source/Developer/ShaderCompilerCommon/Private/HlslLexer.cpp
+++ b/Engine/Source/Developer/ShaderCompilerCommon/Private/HlslLexer.cpp
@@ -191,7 +191,7 @@ namespace CrossCompiler
 				InsertToken(TEXT(">>"), EHlslToken::GreaterGreater);
 				InsertToken(TEXT(">>="), EHlslToken::GreaterGreaterEqual);
 				InsertToken(TEXT("&"), EHlslToken::And);
-				InsertToken(TEXT("&="), EHlslToken::And);
+				InsertToken(TEXT("&="), EHlslToken::AndEqual);
 				InsertToken(TEXT("|"), EHlslToken::Or);
 				InsertToken(TEXT("|="), EHlslToken::OrEqual);
 				InsertToken(TEXT("^"), EHlslToken::Xor);
@@ -335,7 +335,7 @@ namespace CrossCompiler
 				InsertToken(TEXT("float3x4"), EHlslToken::Float3x4);
 				InsertToken(TEXT("float4x4"), EHlslToken::Float4x4);
 
-				InsertToken(TEXT("Texture"), EHlslToken::Texture);
+				InsertToken(TEXT("texture"), EHlslToken::Texture);
 				InsertToken(TEXT("Texture1D"), EHlslToken::Texture1D);
 				InsertToken(TEXT("Texture1DArray"), EHlslToken::Texture1DArray);
 				InsertToken(TEXT("Texture1D_Array"), EHlslToken::Texture1DArray);	// PSSL
@@ -351,11 +351,11 @@ namespace CrossCompiler
 				InsertToken(TEXT("TextureCubeArray"), EHlslToken::TextureCubeArray);
 				InsertToken(TEXT("TextureCube_Array"), EHlslToken::TextureCubeArray);	// PSSL
 
-				InsertToken(TEXT("Sampler"), EHlslToken::Sampler);
-				InsertToken(TEXT("Sampler1D"), EHlslToken::Sampler1D);
-				InsertToken(TEXT("Sampler2D"), EHlslToken::Sampler2D);
-				InsertToken(TEXT("Sampler3D"), EHlslToken::Sampler3D);
-				InsertToken(TEXT("SamplerCube"), EHlslToken::SamplerCube);
+				InsertToken(TEXT("sampler"), EHlslToken::Sampler);
+				InsertToken(TEXT("sampler1D"), EHlslToken::Sampler1D);
+				InsertToken(TEXT("sampler2D"), EHlslToken::Sampler2D);
+				InsertToken(TEXT("sampler3D"), EHlslToken::Sampler3D);
+				InsertToken(TEXT("samplerCUBE"), EHlslToken::SamplerCube);
 				InsertToken(TEXT("SamplerState"), EHlslToken::SamplerState);
 				InsertToken(TEXT("SamplerComparisonState"), EHlslToken::SamplerComparisonState);
 
@@ -385,6 +385,7 @@ namespace CrossCompiler
 				InsertToken(TEXT("RW_Texture3D"), EHlslToken::RWTexture3D);	// PSSL
 				InsertToken(TEXT("StructuredBuffer"), EHlslToken::StructuredBuffer);
 				InsertToken(TEXT("RegularBuffer"), EHlslToken::StructuredBuffer);	// PSSL
+				InsertToken(TEXT("RaytracingAccelerationStructure"), EHlslToken::RaytracingAccelerationStructure);
 
 				// Modifiers
 				InsertToken(TEXT("in"), EHlslToken::In);
@@ -406,6 +407,7 @@ namespace CrossCompiler
 				InsertToken(TEXT("ConstantBuffer"), EHlslToken::CBuffer);	// PSSL
 				InsertToken(TEXT("groupshared"), EHlslToken::GroupShared);
 				InsertToken(TEXT("row_major"), EHlslToken::RowMajor);
+				InsertToken(TEXT("register"), EHlslToken::Register);
 			}
 		} GStaticInitializer;
 	}
@@ -612,7 +614,7 @@ namespace CrossCompiler
 			}
 		}
 
-		bool MatchFloatNumber(float& OutNum)
+		bool MatchFloatNumber(FString& OutLiteral)
 		{
 			auto* Original = Current;
 			TCHAR Char = Peek();
@@ -622,6 +624,7 @@ namespace CrossCompiler
 			// [0-9]+\.[0-9]+([eE][+-]?[0-9]+)?[fF]?	-> Digits+ Dot Digits+ Exp? F?
 			// [0-9]+[eE][+-]?[0-9]+[fF]?				-> Digits+ Exp F?
 			// [0-9]+[fF]								-> Digits+ F
+			// 1.#INF									-> Infinity constant
 			if (!IsDigit(Char) && Char != '.')
 			{
 				return false;
@@ -649,6 +652,19 @@ namespace CrossCompiler
 
 				if (Match('.'))
 				{
+					// Check for infinity constant
+					if (Match('#'))
+					{
+						if (MatchString(TEXT("INF"), 3))
+						{
+							goto Done;
+						}
+						else
+						{
+							goto NotFloat;
+						}
+					}
+
 					bExpOptional = true;
 					MatchAndSkipDigits();
 				}
@@ -698,7 +714,7 @@ namespace CrossCompiler
 			}
 
 		Done:
-			OutNum = FCString::Atof(Original);
+			OutLiteral = FString(static_cast<int32>(Current - Original), Original);
 			return true;
 
 		NotFloat:
@@ -973,10 +989,9 @@ namespace CrossCompiler
 				FString Identifier;
 				EHlslToken SymbolToken;
 				uint32 UnsignedInteger;
-				float FloatNumber;
-				if (Tokenizer.MatchFloatNumber(FloatNumber))
+				if (Tokenizer.MatchFloatNumber(Identifier))
 				{
-					AddToken(FHlslToken(FloatNumber), Tokenizer);
+					AddToken(FHlslToken(EHlslToken::FloatConstant, Identifier), Tokenizer);
 				}
 				else if (Tokenizer.MatchLiteralInteger(UnsignedInteger))
 				{
@@ -1042,7 +1057,7 @@ namespace CrossCompiler
 				break;
 
 			case EHlslToken::FloatConstant:
-				FPlatformMisc::LowLevelOutputDebugStringf(TEXT("** %d: FloatConstant '%f'\n"), Index, Token.Float);
+				FPlatformMisc::LowLevelOutputDebugStringf(TEXT("** %d: FloatConstant '%s'\n"), Index, *Token.String);
 				break;
 
 			default:
diff --git a/Engine/Source/Developer/ShaderCompilerCommon/Private/HlslLexer.h b/Engine/Source/Developer/ShaderCompilerCommon/Private/HlslLexer.h
--- a/Engine/Source/Developer/ShaderCompilerCommon/Private/HlslLexer.h
+++ b/Engine/Source/Developer/ShaderCompilerCommon/Private/HlslLexer.h
@@ -194,7 +194,7 @@ namespace CrossCompiler
 		Float3x4,
 		Float4x4,
 
-		Texture,
+		Texture, // texture (SM3)
 		Texture1D,
 		Texture1DArray,
 		Texture2D,
@@ -205,11 +205,11 @@ namespace CrossCompiler
 		TextureCube,
 		TextureCubeArray,
 
-		Sampler,
-		Sampler1D,
-		Sampler2D,
-		Sampler3D,
-		SamplerCube,
+		Sampler, // sampler (SM3)
+		Sampler1D, // sampler1D (SM3)
+		Sampler2D, // sampler2D (SM3)
+		Sampler3D, // sampler3D (SM3)
+		SamplerCube, // samplerCUBE (SM3)
 		SamplerState,
 		SamplerComparisonState,
 
@@ -226,6 +226,7 @@ namespace CrossCompiler
 		RWTexture2DArray,
 		RWTexture3D,
 		StructuredBuffer,
+		RaytracingAccelerationStructure,
 
 		// Modifiers
 		In,
@@ -246,6 +247,7 @@ namespace CrossCompiler
 		CBuffer,
 		GroupShared,
 		RowMajor,
+		Register,
 
 		Identifier,
 		UnsignedIntegerConstant,
@@ -259,15 +261,13 @@ namespace CrossCompiler
 		EHlslToken Token;
 		FString String;
 		uint32 UnsignedInteger;
-		float Float;
 
 		FSourceInfo SourceInfo;
 
-		explicit FHlslToken(const FString& Identifier) : Token(EHlslToken::Identifier), String(Identifier), UnsignedInteger(0), Float(0) { }
-		explicit FHlslToken(EHlslToken InToken, const FString& Identifier) : Token(InToken), String(Identifier), UnsignedInteger(0), Float(0) { }
-		explicit FHlslToken(uint32 InUnsignedInteger) : Token(EHlslToken::UnsignedIntegerConstant), UnsignedInteger(InUnsignedInteger), Float(0) { }
-		explicit FHlslToken(float InFloat) : Token(EHlslToken::FloatConstant), UnsignedInteger(0), Float(InFloat) { }
-		explicit FHlslToken(bool bInValue) : Token(EHlslToken::BoolConstant), UnsignedInteger(bInValue ? 1 : 0), Float(0) { }
+		explicit FHlslToken(const FString& Identifier) : Token(EHlslToken::Identifier), String(Identifier), UnsignedInteger(0) { }
+		explicit FHlslToken(EHlslToken InToken, const FString& Identifier) : Token(InToken), String(Identifier), UnsignedInteger(0) { }
+		explicit FHlslToken(uint32 InUnsignedInteger) : Token(EHlslToken::UnsignedIntegerConstant), UnsignedInteger(InUnsignedInteger) { }
+		explicit FHlslToken(bool bInValue) : Token(EHlslToken::BoolConstant), UnsignedInteger(bInValue ? 1 : 0) { }
 	};
 
 	class FHlslScanner
diff --git a/Engine/Source/Developer/ShaderCompilerCommon/Private/HlslParser.cpp b/Engine/Source/Developer/ShaderCompilerCommon/Private/HlslParser.cpp
--- a/Engine/Source/Developer/ShaderCompilerCommon/Private/HlslParser.cpp
+++ b/Engine/Source/Developer/ShaderCompilerCommon/Private/HlslParser.cpp
@@ -111,7 +111,7 @@ namespace CrossCompiler
 	{
 		if (Scanner.MatchToken(EHlslToken::LeftSquareBracket))
 		{
-			auto ExpressionResult = ParseExpression(Scanner, SymbolScope, false, Allocator, OutExpression);
+			auto ExpressionResult = ParseExpression(Scanner, SymbolScope, EEF_ALLOW_SEQUENCE, Allocator, OutExpression);
 			if (ExpressionResult == EParseResult::Error)
 			{
 				Scanner.SourceError(TEXT("Expected expression!"));
@@ -170,7 +170,93 @@ namespace CrossCompiler
 		return EParseResult::NotMatched;
 	}
 
-	EParseResult ParseTextureOrBufferSimpleDeclaration(FHlslScanner& Scanner, FSymbolScope* SymbolScope, bool bMultiple, FLinearAllocator* Allocator, AST::FDeclaratorList** OutDeclaratorList)
+	// Multi declaration parser flags
+	enum EDeclarationFlags
+	{
+		EDF_CONST_ROW_MAJOR = (1 << 0),
+		EDF_STATIC = (1 << 1),
+		EDF_UNIFORM = (1 << 2),
+		EDF_TEXTURE_SAMPLER_OR_BUFFER = (1 << 3),
+		EDF_INITIALIZER = (1 << 4),
+		EDF_INITIALIZER_LIST = (1 << 5) | EDF_INITIALIZER,
+		EDF_SEMANTIC = (1 << 6),
+		EDF_SEMICOLON = (1 << 7),
+		EDF_IN_OUT = (1 << 8),
+		EDF_MULTIPLE = (1 << 9),
+		EDF_PRIMITIVE_DATA_TYPE = (1 << 10),
+		EDF_SHARED = (1 << 11),
+		EDF_INTERPOLATION = (1 << 12),
+		EDF_REGISTER = (1 << 13),
+	};
+
+	EParseResult ParseInitializer(FHlslScanner& Scanner, FSymbolScope* SymbolScope, bool bAllowLists, FLinearAllocator* Allocator, AST::FExpression** OutList)
+	{
+		if (bAllowLists && Scanner.MatchToken(EHlslToken::LeftBrace))
+		{
+			*OutList = new(Allocator) AST::FInitializerListExpression(Allocator, Scanner.GetCurrentToken()->SourceInfo);
+			EParseResult Result = ParseExpressionList(EHlslToken::RightBrace, Scanner, SymbolScope, EHlslToken::LeftBrace, Allocator, *OutList);
+			if (Result != EParseResult::Matched)
+			{
+				Scanner.SourceError(TEXT("Invalid initializer list\n"));
+			}
+
+			return EParseResult::Matched;
+		}
+		else
+		{
+			//@todo-LH?
+			EParseResult Result = ParseExpression(Scanner, SymbolScope, EEF_ALLOW_ASSIGNMENT, Allocator, OutList);
+			if (Result == EParseResult::Error)
+			{
+				Scanner.SourceError(TEXT("Invalid initializer expression\n"));
+			}
+
+			return Result;
+		}
+
+		return EParseResult::NotMatched;
+	}
+
+	EParseResult ParseRegisterSpecifier(FHlslScanner& Scanner, FLinearAllocator* Allocator, AST::FRegisterSpecifier** OutSpecifier)
+	{
+		const auto* Token = Scanner.GetCurrentToken();
+		if (Scanner.MatchToken(EHlslToken::Register))
+		{
+			if (!Scanner.MatchToken(EHlslToken::LeftParenthesis))
+			{
+				Scanner.SourceError(TEXT("'(' expected"));
+				return ParseResultError();
+			}
+
+			auto* Specifier = new(Allocator) AST::FRegisterSpecifier(Allocator, Token->SourceInfo);
+
+			// Parse argument expressions of 'register'-specifier
+			do
+			{
+				AST::FExpression* Argument = nullptr;
+				if (ParseExpression(Scanner, nullptr, 0, Allocator, &Argument) != EParseResult::Matched)
+				{
+					Scanner.SourceError(TEXT("Expression expected"));
+					return ParseResultError();
+				}
+				Specifier->Arguments.Add(Argument);
+			}
+			while (Scanner.MatchToken(EHlslToken::Comma));
+
+			if (!Scanner.MatchToken(EHlslToken::RightParenthesis))
+			{
+				Scanner.SourceError(TEXT("')' expected"));
+				return ParseResultError();
+			}
+
+			*OutSpecifier = Specifier;
+			return EParseResult::Matched;
+		}
+
+		return EParseResult::NotMatched;
+	}
+
+	EParseResult ParseTextureOrBufferSimpleDeclaration(FHlslScanner& Scanner, FSymbolScope* SymbolScope, bool bMultiple, bool bInitializer, bool bRegister, FLinearAllocator* Allocator, AST::FDeclaratorList** OutDeclaratorList)
 	{
 		auto OriginalToken = Scanner.GetCurrentTokenIndex();
 		const auto* Token = Scanner.GetCurrentToken();
@@ -238,6 +324,33 @@ namespace CrossCompiler
 					return ParseResultError();
 				}
 
+				// Parse optional 'register'-specifier
+				if (bRegister)
+				{
+					if (Scanner.MatchToken(EHlslToken::Colon))
+					{
+						if (ParseRegisterSpecifier(Scanner, Allocator, &Declaration->Register) != EParseResult::Matched)
+						{
+							Scanner.SourceError(TEXT("Invalid register specifier"));
+							return ParseResultError();
+						}
+					}
+				}
+
+				// Parse optional initializer
+				if (bInitializer)
+				{
+					if (Scanner.MatchToken(EHlslToken::Equal))
+					{
+						const bool bAllowInitializerList = false;
+						if (ParseInitializer(Scanner, SymbolScope, bAllowInitializerList, Allocator, &Declaration->Initializer) != EParseResult::Matched)
+						{
+							Scanner.SourceError(TEXT("Invalid initializer\n"));
+							return ParseResultError();
+						}
+					}
+				}
+
 				(*OutDeclaratorList)->Declarations.Add(Declaration);
 			}
 			while (bMultiple && Scanner.MatchToken(EHlslToken::Comma));
@@ -250,52 +363,6 @@ namespace CrossCompiler
 		return EParseResult::NotMatched;
 	}
 
-	// Multi declaration parser flags
-	enum EDeclarationFlags
-	{
-		EDF_CONST_ROW_MAJOR				= (1 << 0),
-		EDF_STATIC						= (1 << 1),
-		EDF_UNIFORM						= (1 << 2),
-		EDF_TEXTURE_SAMPLER_OR_BUFFER	= (1 << 3),
-		EDF_INITIALIZER					= (1 << 4),
-		EDF_INITIALIZER_LIST			= (1 << 5) | EDF_INITIALIZER,
-		EDF_SEMANTIC					= (1 << 6),
-		EDF_SEMICOLON					= (1 << 7),
-		EDF_IN_OUT						= (1 << 8),
-		EDF_MULTIPLE					= (1 << 9),
-		EDF_PRIMITIVE_DATA_TYPE			= (1 << 10),
-		EDF_SHARED						= (1 << 11),
-		EDF_INTERPOLATION				= (1 << 12),
-	};
-
-	EParseResult ParseInitializer(FHlslScanner& Scanner, FSymbolScope* SymbolScope, bool bAllowLists, FLinearAllocator* Allocator, AST::FExpression** OutList)
-	{
-		if (bAllowLists && Scanner.MatchToken(EHlslToken::LeftBrace))
-		{
-			*OutList = new(Allocator) AST::FInitializerListExpression(Allocator, Scanner.GetCurrentToken()->SourceInfo);
-			auto Result = ParseExpressionList(EHlslToken::RightBrace, Scanner, SymbolScope, EHlslToken::LeftBrace, Allocator, *OutList);
-			if (Result != EParseResult::Matched)
-			{
-				Scanner.SourceError(TEXT("Invalid initializer list\n"));
-			}
-
-			return EParseResult::Matched;
-		}
-		else
-		{
-			//@todo-rco?
-			auto Result = ParseExpression(Scanner, SymbolScope, true, Allocator, OutList);
-			if (Result == EParseResult::Error)
-			{
-				Scanner.SourceError(TEXT("Invalid initializer expression\n"));
-			}
-
-			return Result;
-		}
-
-		return EParseResult::NotMatched;
-	}
-
 	EParseResult ParseDeclarationStorageQualifiers(FHlslScanner& Scanner, int32 TypeFlags, int32 DeclarationFlags, bool& bOutPrimitiveFound, AST::FTypeQualifier* Qualifier)
 	{
 		bOutPrimitiveFound = false;
@@ -627,7 +694,11 @@ namespace CrossCompiler
 
 		if (DeclarationFlags & EDF_TEXTURE_SAMPLER_OR_BUFFER)
 		{
-			auto Result = ParseTextureOrBufferSimpleDeclaration(Scanner, SymbolScope, (DeclarationFlags & EDF_MULTIPLE) == EDF_MULTIPLE, Allocator, &DeclaratorList);
+			bool bMultiple = ((DeclarationFlags & EDF_MULTIPLE) == EDF_MULTIPLE);
+			bool bInitializer = ((DeclarationFlags & EDF_INITIALIZER) == EDF_INITIALIZER);
+			bool bRegister = ((DeclarationFlags & EDF_REGISTER) == EDF_REGISTER);
+
+			auto Result = ParseTextureOrBufferSimpleDeclaration(Scanner, SymbolScope, bMultiple, bInitializer, bRegister, Allocator, &DeclaratorList);
 			if (Result == EParseResult::Matched)
 			{
 				*OutDeclaratorList = DeclaratorList;
@@ -873,7 +944,7 @@ namespace CrossCompiler
 			}
 
 			AST::FDeclaratorList* Declaration = nullptr;
-			auto Result = ParseGeneralDeclaration(Scanner, SymbolScope, Allocator, &Declaration, 0, EDF_CONST_ROW_MAJOR | EDF_SEMICOLON | EDF_SEMANTIC | EDF_TEXTURE_SAMPLER_OR_BUFFER | EDF_INTERPOLATION);
+			auto Result = ParseGeneralDeclaration(Scanner, SymbolScope, Allocator, &Declaration, 0, EDF_CONST_ROW_MAJOR | EDF_SEMICOLON | EDF_SEMANTIC | EDF_MULTIPLE | EDF_TEXTURE_SAMPLER_OR_BUFFER | EDF_INTERPOLATION);
 			if (Result == EParseResult::Error)
 			{
 				return ParseResultError();
@@ -1115,7 +1186,10 @@ check(0);
 	EParseResult ParseLocalDeclaration(FHlslParser& Parser, FLinearAllocator* Allocator, AST::FNode** OutDeclaration)
 	{
 		AST::FDeclaratorList* List = nullptr;
-		auto Result = ParseGeneralDeclaration(Parser.Scanner, Parser.CurrentScope, Allocator, &List, 0, EDF_CONST_ROW_MAJOR | EDF_INITIALIZER | EDF_INITIALIZER_LIST | EDF_SEMICOLON | EDF_MULTIPLE | EDF_STATIC);
+		auto Result = ParseGeneralDeclaration(
+			Parser.Scanner, Parser.CurrentScope, Allocator, &List, 0,
+			EDF_CONST_ROW_MAJOR | EDF_TEXTURE_SAMPLER_OR_BUFFER | EDF_INITIALIZER | EDF_INITIALIZER_LIST | EDF_SEMICOLON | EDF_MULTIPLE | EDF_STATIC
+		);
 		*OutDeclaration = List;
 		return Result;
 	}
@@ -1123,7 +1197,10 @@ check(0);
 	EParseResult ParseGlobalVariableDeclaration(FHlslParser& Parser, FLinearAllocator* Allocator, AST::FNode** OutDeclaration)
 	{
 		AST::FDeclaratorList* List = nullptr;
-		auto Result = ParseGeneralDeclaration(Parser.Scanner, Parser.CurrentScope, Allocator, &List, ETF_USER_TYPES | ETF_ERROR_IF_NOT_USER_TYPE, EDF_CONST_ROW_MAJOR | EDF_STATIC | EDF_SHARED | EDF_TEXTURE_SAMPLER_OR_BUFFER | EDF_INITIALIZER | EDF_INITIALIZER_LIST | EDF_SEMICOLON | EDF_MULTIPLE | EDF_UNIFORM | EDF_INTERPOLATION);
+		auto Result = ParseGeneralDeclaration(
+			Parser.Scanner, Parser.CurrentScope, Allocator, &List, ETF_USER_TYPES | ETF_ERROR_IF_NOT_USER_TYPE,
+			EDF_CONST_ROW_MAJOR | EDF_STATIC | EDF_SHARED | EDF_TEXTURE_SAMPLER_OR_BUFFER | EDF_INITIALIZER | EDF_INITIALIZER_LIST | EDF_SEMICOLON | EDF_MULTIPLE | EDF_UNIFORM | EDF_INTERPOLATION | EDF_REGISTER
+		);
 		*OutDeclaration = List;
 		return Result;
 	}
@@ -1138,7 +1215,7 @@ check(0);
 			return EParseResult::Matched;
 		}
 
-		if (ParseExpression(Parser.Scanner, Parser.CurrentScope, true, Allocator, &Statement->OptionalExpression) != EParseResult::Matched)
+		if (ParseExpression(Parser.Scanner, Parser.CurrentScope, (EEF_ALLOW_ASSIGNMENT | EEF_ALLOW_SEQUENCE), Allocator, &Statement->OptionalExpression) != EParseResult::Matched)
 		{
 			Parser.Scanner.SourceError(TEXT("Expression expected"));
 			return ParseResultError();
@@ -1178,7 +1255,7 @@ check(0);
 		}
 
 		AST::FExpression* ConditionExpression = nullptr;
-		if (ParseExpression(Parser.Scanner, Parser.CurrentScope, true, Allocator, &ConditionExpression) != EParseResult::Matched)
+		if (ParseExpression(Parser.Scanner, Parser.CurrentScope, (EEF_ALLOW_ASSIGNMENT | EEF_ALLOW_SEQUENCE), Allocator, &ConditionExpression) != EParseResult::Matched)
 		{
 			Parser.Scanner.SourceError(TEXT("Expression expected"));
 			return ParseResultError();
@@ -1214,7 +1291,7 @@ check(0);
 		}
 
 		AST::FExpression* ConditionExpression = nullptr;
-		if (ParseExpression(Parser.Scanner, Parser.CurrentScope, true, Allocator, &ConditionExpression) != EParseResult::Matched)
+		if (ParseExpression(Parser.Scanner, Parser.CurrentScope, (EEF_ALLOW_ASSIGNMENT | EEF_ALLOW_SEQUENCE), Allocator, &ConditionExpression) != EParseResult::Matched)
 		{
 			Parser.Scanner.SourceError(TEXT("Expression expected"));
 			return ParseResultError();
@@ -1275,7 +1352,7 @@ check(0);
 		}
 
 		AST::FExpression* ConditionExpression = nullptr;
-		auto Result = ParseExpression(Parser.Scanner, Parser.CurrentScope, true, Allocator, &ConditionExpression);
+		auto Result = ParseExpression(Parser.Scanner, Parser.CurrentScope, (EEF_ALLOW_ASSIGNMENT | EEF_ALLOW_SEQUENCE), Allocator, &ConditionExpression);
 		if (Result == EParseResult::Error)
 		{
 			Parser.Scanner.SourceError(TEXT("Expected expression or declaration!\n"));
@@ -1289,7 +1366,7 @@ check(0);
 		}
 
 		AST::FExpression* RestExpression = nullptr;
-		Result = ParseExpression(Parser.Scanner, Parser.CurrentScope, true, Allocator, &RestExpression);
+		Result = ParseExpression(Parser.Scanner, Parser.CurrentScope, (EEF_ALLOW_ASSIGNMENT | EEF_ALLOW_SEQUENCE), Allocator, &RestExpression);
 		if (Result == EParseResult::Error)
 		{
 			Parser.Scanner.SourceError(TEXT("Expected expression or declaration!\n"));
@@ -1330,7 +1407,7 @@ check(0);
 			return ParseResultError();
 		}
 
-		if (ParseExpression(Parser.Scanner, Parser.CurrentScope, true, Allocator, &Statement->Condition) != EParseResult::Matched)
+		if (ParseExpression(Parser.Scanner, Parser.CurrentScope, (EEF_ALLOW_ASSIGNMENT | EEF_ALLOW_SEQUENCE), Allocator, &Statement->Condition) != EParseResult::Matched)
 		{
 			Parser.Scanner.SourceError(TEXT("Expression expected"));
 			return ParseResultError();
@@ -1385,7 +1462,7 @@ check(0);
 				else
 				{
 					AST::FExpression* Expression = nullptr;
-					EParseResult Result = ParseExpression(Scanner, SymbolScope, false, Allocator, &Expression);
+					EParseResult Result = ParseExpression(Scanner, SymbolScope, 0, Allocator, &Expression);
 					if (Result != EParseResult::Matched)
 					{
 						Scanner.SourceError(TEXT("Incorrect attribute expression!\n"));
@@ -1505,7 +1582,7 @@ check(0);
 				else if (Parser.Scanner.MatchToken(EHlslToken::Case))
 				{
 					AST::FExpression* CaseExpression = nullptr;
-					if (ParseExpression(Parser.Scanner, Parser.CurrentScope, true, Allocator, &CaseExpression) != EParseResult::Matched)
+					if (ParseExpression(Parser.Scanner, Parser.CurrentScope, (EEF_ALLOW_ASSIGNMENT | EEF_ALLOW_SEQUENCE), Allocator, &CaseExpression) != EParseResult::Matched)
 					{
 						Parser.Scanner.SourceError(TEXT("Expression expected on case label!"));
 						return ParseResultError();
@@ -1585,7 +1662,7 @@ check(0);
 		}
 
 		AST::FExpression* Condition = nullptr;
-		if (ParseExpression(Parser.Scanner, Parser.CurrentScope, false, Allocator, &Condition) != EParseResult::Matched)
+		if (ParseExpression(Parser.Scanner, Parser.CurrentScope, EEF_ALLOW_SEQUENCE, Allocator, &Condition) != EParseResult::Matched)
 		{
 			Parser.Scanner.SourceError(TEXT("Expression expected"));
 			return ParseResultError();
@@ -1613,7 +1690,7 @@ check(0);
 	{
 		auto OriginalToken = Parser.Scanner.GetCurrentTokenIndex();
 		auto* Statement = new(Allocator) AST::FExpressionStatement(Allocator, nullptr, Parser.Scanner.GetCurrentToken()->SourceInfo);
-		if (ParseExpression(Parser.Scanner, Parser.CurrentScope, true, Allocator, &Statement->Expression) == EParseResult::Matched)
+		if (ParseExpression(Parser.Scanner, Parser.CurrentScope, (EEF_ALLOW_ASSIGNMENT | EEF_ALLOW_SEQUENCE), Allocator, &Statement->Expression) == EParseResult::Matched)
 		{
 			if (Parser.Scanner.MatchToken(EHlslToken::Semicolon))
 			{
@@ -1741,7 +1818,7 @@ check(0);
 				RulesStatements.Add(FRulePair(EHlslToken::Semicolon, ParseEmptyStatement));
 				RulesStatements.Add(FRulePair(EHlslToken::Break, ParseBreakStatement));
 				RulesStatements.Add(FRulePair(EHlslToken::Continue, ParseContinueStatement));
-				RulesStatements.Add(FRulePair(EHlslToken::Invalid, ParseLocalDeclaration));
+				RulesStatements.Add(FRulePair(EHlslToken::Invalid, ParseLocalDeclaration, true));
 				// Always try expressions last
 				RulesStatements.Add(FRulePair(EHlslToken::Invalid, ParseExpressionStatement));
 			}
@@ -1774,6 +1851,9 @@ check(0);
 			//Found = Found->FindNamespace(TEXT("Gnm"));
 			//Found->FindType(Found, TEXT("Sampler"), false);
 		}
+
+		// Register built-in structure for DXR in SM6
+		CurrentScope->Add(TEXT("RayDesc"));
 	}
 
 	namespace Parser
diff --git a/Engine/Source/Developer/ShaderCompilerCommon/Private/HlslParser.h b/Engine/Source/Developer/ShaderCompilerCommon/Private/HlslParser.h
--- a/Engine/Source/Developer/ShaderCompilerCommon/Private/HlslParser.h
+++ b/Engine/Source/Developer/ShaderCompilerCommon/Private/HlslParser.h
@@ -31,12 +31,12 @@ namespace CrossCompiler
 		typedef void TCallback(void* CallbackData, CrossCompiler::FLinearAllocator* Allocator, CrossCompiler::TLinearArray<CrossCompiler::AST::FNode*>& ASTNodes);
 
 		// Returns true if successfully parsed
-		bool Parse(const FString& Input, const FString& Filename, FCompilerMessages& OutCompilerMessages, TCallback* Callback, void* CallbackData = nullptr);
+		SHADERCOMPILERCOMMON_API bool Parse(const FString& Input, const FString& Filename, FCompilerMessages& OutCompilerMessages, TCallback* Callback, void* CallbackData = nullptr);
 
 		// Returns true if successfully parsed
-		bool Parse(const FString& Input, const FString& Filename, FCompilerMessages& OutCompilerMessages, TFunction< void(CrossCompiler::FLinearAllocator* Allocator, CrossCompiler::TLinearArray<CrossCompiler::AST::FNode*>& ASTNodes)> Function);
+		SHADERCOMPILERCOMMON_API bool Parse(const FString& Input, const FString& Filename, FCompilerMessages& OutCompilerMessages, TFunction< void(CrossCompiler::FLinearAllocator* Allocator, CrossCompiler::TLinearArray<CrossCompiler::AST::FNode*>& ASTNodes)> Function);
 
 		// Sample callback to write out all nodes into a string; pass a valid pointer FString* as OutFStringPointer
-		void WriteNodesToString(void* OutFStringPointer, CrossCompiler::FLinearAllocator* Allocator, CrossCompiler::TLinearArray<CrossCompiler::AST::FNode*>& ASTNodes);
+		SHADERCOMPILERCOMMON_API void WriteNodesToString(void* OutFStringPointer, CrossCompiler::FLinearAllocator* Allocator, CrossCompiler::TLinearArray<CrossCompiler::AST::FNode*>& ASTNodes);
 	}
 }
diff --git a/Engine/Source/Developer/ShaderCompilerCommon/Private/HlslUtils.cpp b/Engine/Source/Developer/ShaderCompilerCommon/Private/HlslUtils.cpp
--- a/Engine/Source/Developer/ShaderCompilerCommon/Private/HlslUtils.cpp
+++ b/Engine/Source/Developer/ShaderCompilerCommon/Private/HlslUtils.cpp
@@ -1071,7 +1071,7 @@ struct FRemoveUnusedInputs : FRemoveAlgorithm
 		using namespace CrossCompiler::AST;
 
 		auto* Zero = new(Allocator) FUnaryExpression(Allocator, EOperators::FloatConstant, nullptr, SourceInfo);
-		Zero->FloatConstant = 0;
+		Zero->Identifier = Allocator->Strdup(TEXT("0"));
 		auto* Initializer = new(Allocator) FUnaryExpression(Allocator, EOperators::TypeCast, Zero, SourceInfo);
 		Initializer->TypeSpecifier = MakeSimpleType(OriginalStructSpecifier->Name)->Specifier;
 
@@ -1137,7 +1137,7 @@ struct FRemoveUnusedInputs : FRemoveAlgorithm
 		{
 			// Make a local to generate the in parameter: Type Local = (Type)0;
 			auto* Zero = new(Allocator) FUnaryExpression(Allocator, EOperators::FloatConstant, nullptr, SourceInfo);
-			Zero->FloatConstant = 0;
+			Zero->Identifier = Allocator->Strdup(TEXT("0"));
 			Initializer = new(Allocator) FUnaryExpression(Allocator, EOperators::TypeCast, Zero, SourceInfo);
 			Initializer->TypeSpecifier = ParameterDeclarator->Type->Specifier;
 
