Removed debugging code of "check(0)" on HLSL forward declarations in HlslParser.

#rnx

diff --git a/Engine/Source/Developer/ShaderCompilerCommon/Private/HlslParser.cpp b/Engine/Source/Developer/ShaderCompilerCommon/Private/HlslParser.cpp
--- a/Engine/Source/Developer/ShaderCompilerCommon/Private/HlslParser.cpp
+++ b/Engine/Source/Developer/ShaderCompilerCommon/Private/HlslParser.cpp
@@ -1140,7 +1140,6 @@ Done:
 
 		if (Parser.Scanner.MatchToken(EHlslToken::Semicolon))
 		{
-check(0);
 			// Forward declare
 			return EParseResult::Matched;
 		}
