Fixed compiler warning for PVS-Studio in while-loop statement.

#rnx

diff --git a/Engine/Source/Developer/ShaderCompilerCommon/Private/HlslExpressionParser.inl b/Engine/Source/Developer/ShaderCompilerCommon/Private/HlslExpressionParser.inl
--- a/Engine/Source/Developer/ShaderCompilerCommon/Private/HlslExpressionParser.inl
+++ b/Engine/Source/Developer/ShaderCompilerCommon/Private/HlslExpressionParser.inl
@@ -547,7 +547,7 @@ namespace CrossCompiler
 				const auto* Peek2 = Scanner.PeekToken(2);
 
 				bool bFoundConst = false;
-				int PeekN = 0;
+				int32 PeekN = 0;
 				auto HandleUnaryToken = [&](EHlslToken TokenType)
 				{
 					if (Peek1->Token == TokenType)
@@ -575,10 +575,9 @@ namespace CrossCompiler
 				//#todo-rco: Is precise allowed on casts?
 				if (Peek1 && ParseGeneralTypeFromToken(Peek1, ETF_BUILTIN_NUMERIC | ETF_USER_TYPES, false, SymbolScope, Allocator, &TypeSpecifier) == EParseResult::Matched && Peek2 && Peek2->Token == EHlslToken::RightParenthesis)
 				{
-					while (PeekN > 0)
+					for (; PeekN > 0; --PeekN)
 					{
 						Scanner.Advance();
-						--PeekN;
 					}
 
 					// Cast
