Suppress false alarm of PVS-Studio warning V654 in HlslExpressionParser.inl.

#rnx

diff --git a/Engine/Source/Developer/ShaderCompilerCommon/Private/HlslExpressionParser.inl b/Engine/Source/Developer/ShaderCompilerCommon/Private/HlslExpressionParser.inl
--- a/Engine/Source/Developer/ShaderCompilerCommon/Private/HlslExpressionParser.inl
+++ b/Engine/Source/Developer/ShaderCompilerCommon/Private/HlslExpressionParser.inl
@@ -575,7 +575,7 @@ namespace CrossCompiler
 				//#todo-rco: Is precise allowed on casts?
 				if (Peek1 && ParseGeneralTypeFromToken(Peek1, ETF_BUILTIN_NUMERIC | ETF_USER_TYPES, false, SymbolScope, Allocator, &TypeSpecifier) == EParseResult::Matched && Peek2 && Peek2->Token == EHlslToken::RightParenthesis)
 				{
-					for (; PeekN > 0; --PeekN)
+					for (; PeekN > 0; --PeekN) //-V654
 					{
 						Scanner.Advance();
 					}
