UE4R - Support precise on HlslParser (missing file)
#rnx

diff --git a/Engine/Source/Developer/ShaderCompilerCommon/Private/HlslLexer.cpp b/Engine/Source/Developer/ShaderCompilerCommon/Private/HlslLexer.cpp
--- a/Engine/Source/Developer/ShaderCompilerCommon/Private/HlslLexer.cpp
+++ b/Engine/Source/Developer/ShaderCompilerCommon/Private/HlslLexer.cpp
@@ -224,6 +224,7 @@ namespace CrossCompiler
 				// Types
 				InsertToken(TEXT("void"), EHlslToken::Void);
 				InsertToken(TEXT("const"), EHlslToken::Const);
+				InsertToken(TEXT("precise"), EHlslToken::Precise);
 
 				InsertToken(TEXT("bool"), EHlslToken::Bool);
 				InsertToken(TEXT("bool1"), EHlslToken::Bool1);
@@ -1178,7 +1179,7 @@ namespace CrossCompiler
 
 				if (Tokenizer.MatchString(MATCH_TARGET(TEXT("#endif"))))
 				{
-					Tokenizer.SkipToNextLine();
+					// Nothing here, skip to next line will happen later
 				}
 				else
 				{
@@ -1194,7 +1195,7 @@ namespace CrossCompiler
 		else
 		{
 			FString Directive = TEXT("#") + Tokenizer.ReadToEndOfLine();
-			CompilerMessages.SourceWarning(*FString::Printf(TEXT("Unhandled preprocessor directive (%s); HlslParser requires preprocessed input!"), Tokenizer.Current));
+			CompilerMessages.SourceWarning(*FString::Printf(TEXT("Unhandled token (%s); HlslParser requires preprocessed input!"), Tokenizer.Current));
 		}
 
 		Tokenizer.SkipToNextLine();
