UE4R - Fix forward declarations on HlslParser

diff --git a/Engine/Source/Developer/ShaderCompilerCommon/Private/HlslParser.cpp b/Engine/Source/Developer/ShaderCompilerCommon/Private/HlslParser.cpp
--- a/Engine/Source/Developer/ShaderCompilerCommon/Private/HlslParser.cpp
+++ b/Engine/Source/Developer/ShaderCompilerCommon/Private/HlslParser.cpp
@@ -1140,7 +1140,8 @@ Done:
 
 		if (Parser.Scanner.MatchToken(EHlslToken::Semicolon))
 		{
-			// Forward declare
+			// Forward declaration
+			*OutFunction = Function;
 			return EParseResult::Matched;
 		}
 		else
@@ -1912,10 +1913,10 @@ Done:
 
 			bool bSuccess = true;
 			TLinearArray<AST::FNode*> Nodes(&Allocator);
+			AST::FNode* PrevNode = nullptr;
 			while (Parser.Scanner.HasMoreTokens())
 			{
 				auto LastIndex = Parser.Scanner.GetCurrentTokenIndex();
-
 				static FString GlobalDeclOrDefinition(TEXT("Global declaration or definition"));
 				AST::FNode* Node = nullptr;
 				auto Result = TryTranslationUnit(Parser, &Allocator, &Node);
@@ -1933,10 +1934,14 @@ Done:
 					Node->Dump(0);
 					}
 					*/
+					checkf(Node, TEXT("Null AST Node! PrevNode @ %s(%d)"),
+						((PrevNode && PrevNode->SourceInfo.Filename) ? *(*PrevNode->SourceInfo.Filename) : TEXT("<none>")),
+						(PrevNode ? PrevNode->SourceInfo.Line : 0));
 					Nodes.Add(Node);
 				}
 
 				check(LastIndex != Parser.Scanner.GetCurrentTokenIndex());
+				PrevNode = Node;
 			}
 
 			if (bSuccess)
