FNM - Copy from 9961559, 9961651, 9961760, 9969176: HlslParser fixes
- Support for primitive type in geometry shaders
- Support for unorm/snorm inside buffer inner types
- Support for inline
- Fix HlslLexer eating the line after #pragma
- Fix forward function declarations not written correctly
- Support for typedef on HlslParser
#rnx

diff --git a/Engine/Source/Developer/ShaderCompilerCommon/Private/HlslAST.cpp b/Engine/Source/Developer/ShaderCompilerCommon/Private/HlslAST.cpp
--- a/Engine/Source/Developer/ShaderCompilerCommon/Private/HlslAST.cpp
+++ b/Engine/Source/Developer/ShaderCompilerCommon/Private/HlslAST.cpp
@@ -591,7 +591,14 @@ namespace CrossCompiler
 			{
 				Writer << TEXT(" : ") << ReturnSemantic;
 			}
-			Writer << TEXT("\n");
+			if (bIsDefinition)
+			{
+				Writer << TEXT(";\n");
+			}
+			else
+			{
+				Writer << TEXT("\n");
+			}
 		}
 
 		FFunction::~FFunction()
@@ -847,6 +854,12 @@ namespace CrossCompiler
 			{
 				Writer << TEXT("row_major ");
 			}
+
+			if (PrimitiveType)
+			{
+				Writer << PrimitiveType;
+				Writer << (TCHAR)' ';
+			}
 		}
 
 		FFullySpecifiedType::FFullySpecifiedType(FLinearAllocator* InAllocator, const FSourceInfo& InInfo) :
@@ -954,6 +967,11 @@ namespace CrossCompiler
 		{
 			Writer.DoIndent();
 			WriteAttributes(Writer);
+			if (bTypedef)
+			{
+				Writer << TEXT("typedef ");
+			}
+
 			if (Type)
 			{
 				Type->Write(Writer);
@@ -1117,7 +1135,6 @@ namespace CrossCompiler
 				Writer << TEXT(")\n");
 				if (Body)
 				{
-					FASTWriterIncrementScope Scope(Writer);
 					Body->Write(Writer);
 				}
 				else
diff --git a/Engine/Source/Developer/ShaderCompilerCommon/Private/HlslAST.h b/Engine/Source/Developer/ShaderCompilerCommon/Private/HlslAST.h
--- a/Engine/Source/Developer/ShaderCompilerCommon/Private/HlslAST.h
+++ b/Engine/Source/Developer/ShaderCompilerCommon/Private/HlslAST.h
@@ -518,6 +518,8 @@ namespace CrossCompiler
 				uint32 Raw;
 			};
 
+			const TCHAR* PrimitiveType = nullptr;
+
 			FTypeQualifier();
 
 			void Write(FASTWriter& Writer) const;
@@ -589,6 +591,8 @@ namespace CrossCompiler
 
 			FFullySpecifiedType* Type;
 			TLinearArray<FNode*> Declarations;
+
+			bool bTypedef = false;
 		};
 
 		struct FParameterDeclarator : public FNode
@@ -623,7 +627,7 @@ namespace CrossCompiler
 
 			TLinearArray<FNode*> Parameters;
 
-			bool bIsDefinition;
+			bool bIsDefinition = false;
 
 			//Signature
 		};
diff --git a/Engine/Source/Developer/ShaderCompilerCommon/Private/HlslExpressionParser.inl b/Engine/Source/Developer/ShaderCompilerCommon/Private/HlslExpressionParser.inl
--- a/Engine/Source/Developer/ShaderCompilerCommon/Private/HlslExpressionParser.inl
+++ b/Engine/Source/Developer/ShaderCompilerCommon/Private/HlslExpressionParser.inl
@@ -170,6 +170,14 @@ namespace CrossCompiler
 		ETF_SAMPLER_TEXTURE_BUFFER	= 1 << 2,
 		ETF_USER_TYPES				= 1 << 3,
 		ETF_ERROR_IF_NOT_USER_TYPE	= 1 << 4,
+		ETF_UNORM					= 1 << 5,
+	};
+
+	enum EExtraQualifiers
+	{
+		EEQ_PRECISE		= 1 << 0,
+		EEQ_UNORM		= 1 << 1,
+		EEQ_SNORM		= 1 << 2,
 	};
 
 	enum EExpressionFlags
@@ -178,7 +186,7 @@ namespace CrossCompiler
 		EEF_ALLOW_SEQUENCE		= 1 << 1,
 	};
 
-	EParseResult ParseGeneralType(const FHlslToken* Token, int32 TypeFlags, bool bPrecise, FLinearAllocator* Allocator, AST::FTypeSpecifier** OutSpecifier)
+	EParseResult ParseGeneralTypeToken(const FHlslToken* Token, int32 TypeFlags, int32 ExtraQualifierFlags, FLinearAllocator* Allocator, AST::FTypeSpecifier** OutSpecifier)
 	{
 		if (!Token)
 		{
@@ -360,10 +368,18 @@ namespace CrossCompiler
 		if (bMatched)
 		{
 			//#todo-rco: Don't re-allocate types
+			FString TypeName;
+			TypeName += (ExtraQualifierFlags & EEQ_SNORM) == EEQ_SNORM
+				? TEXT("snorm ")
+				: ((ExtraQualifierFlags & EEQ_UNORM) == EEQ_UNORM
+					? TEXT("unorm ")
+					: TEXT(""));
+			TypeName += Token->String;
+
 			auto* Type = new(Allocator) AST::FTypeSpecifier(Allocator, Token->SourceInfo);
-			Type->TypeName = Allocator->Strdup(Token->String);
+			Type->TypeName = Allocator->Strdup(TypeName);
 			Type->InnerType = InnerType;
-			Type->bPrecise = bPrecise;
+			Type->bPrecise = (ExtraQualifierFlags & EEQ_PRECISE) == EEQ_PRECISE;
 			*OutSpecifier = Type;
 			return EParseResult::Matched;
 		}
@@ -371,11 +387,11 @@ namespace CrossCompiler
 		return EParseResult::NotMatched;
 	}
 
-	EParseResult ParseGeneralTypeFromToken(const FHlslToken* Token, int32 TypeFlags, bool bPrecise, FSymbolScope* SymbolScope, FLinearAllocator* Allocator, AST::FTypeSpecifier** OutSpecifier)
+	EParseResult ParseGeneralTypeFromToken(const FHlslToken* Token, int32 TypeFlags, int32 ExtraQualifierFlags, FSymbolScope* SymbolScope, FLinearAllocator* Allocator, AST::FTypeSpecifier** OutSpecifier)
 	{
 		if (Token)
 		{
-			if (ParseGeneralType(Token, TypeFlags, bPrecise, Allocator, OutSpecifier) == EParseResult::Matched)
+			if (ParseGeneralTypeToken(Token, TypeFlags, ExtraQualifierFlags, Allocator, OutSpecifier) == EParseResult::Matched)
 			{
 				return EParseResult::Matched;
 			}
@@ -468,14 +484,31 @@ namespace CrossCompiler
 			}
 		}
 
+		int32 ExtraQualifiers = 0;
 		bool bPrecise = false;
 		if (Scanner.MatchToken(EHlslToken::Precise))
 		{
 			Token = Scanner.GetCurrentToken();
-			bPrecise = true;
+			ExtraQualifiers |= EEQ_PRECISE;
+		}
+
+		if ((TypeFlags & ETF_UNORM) == ETF_UNORM && Token->String.Len() == 5)
+		{
+			if (!FCString::Strcmp(*Token->String, TEXT("unorm")))
+			{
+				ExtraQualifiers |= EEQ_UNORM;
+				Scanner.Advance();
+				Token = Scanner.GetCurrentToken();
+			}
+			else if (!FCString::Strcmp(*Token->String, TEXT("snorm")))
+			{
+				ExtraQualifiers |= EEQ_SNORM;
+				Scanner.Advance();
+				Token = Scanner.GetCurrentToken();
+			}
 		}
 
-		auto Result = ParseGeneralTypeFromToken(Token, TypeFlags, bPrecise, SymbolScope, Allocator, OutSpecifier);
+		auto Result = ParseGeneralTypeFromToken(Token, TypeFlags, ExtraQualifiers, SymbolScope, Allocator, OutSpecifier);
 		if (Result == EParseResult::Matched)
 		{
 			Scanner.Advance();
@@ -573,7 +606,7 @@ namespace CrossCompiler
 
 				AST::FTypeSpecifier* TypeSpecifier = nullptr;
 				//#todo-rco: Is precise allowed on casts?
-				if (Peek1 && ParseGeneralTypeFromToken(Peek1, ETF_BUILTIN_NUMERIC | ETF_USER_TYPES, false, SymbolScope, Allocator, &TypeSpecifier) == EParseResult::Matched && Peek2 && Peek2->Token == EHlslToken::RightParenthesis)
+				if (Peek1 && ParseGeneralTypeFromToken(Peek1, ETF_BUILTIN_NUMERIC | ETF_USER_TYPES, 0, SymbolScope, Allocator, &TypeSpecifier) == EParseResult::Matched && Peek2 && Peek2->Token == EHlslToken::RightParenthesis)
 				{
 					for (; PeekN > 0; --PeekN) //-V::654,621
 					{
diff --git a/Engine/Source/Developer/ShaderCompilerCommon/Private/HlslLexer.cpp b/Engine/Source/Developer/ShaderCompilerCommon/Private/HlslLexer.cpp
--- a/Engine/Source/Developer/ShaderCompilerCommon/Private/HlslLexer.cpp
+++ b/Engine/Source/Developer/ShaderCompilerCommon/Private/HlslLexer.cpp
@@ -409,6 +409,8 @@ namespace CrossCompiler
 				InsertToken(TEXT("groupshared"), EHlslToken::GroupShared);
 				InsertToken(TEXT("row_major"), EHlslToken::RowMajor);
 				InsertToken(TEXT("register"), EHlslToken::Register);
+				InsertToken(TEXT("inline"), EHlslToken::Inline);
+				InsertToken(TEXT("typedef"), EHlslToken::Typedef);
 			}
 		} GStaticInitializer;
 	}
@@ -794,7 +796,7 @@ namespace CrossCompiler
 
 		static void ProcessDirective(FTokenizer& Tokenizer, FCompilerMessages& CompilerMessages, class FHlslScanner& Scanner);
 
-		FString ReadToEndOfLine()
+		FString ReadToEndOfLine(bool bSkipToNextLine = true)
 		{
 			FString String;
 			const TCHAR* Start = Current;
@@ -817,7 +819,11 @@ namespace CrossCompiler
 					++Current;
 				}
 			}
-			SkipToNextLine();
+
+			if (bSkipToNextLine)
+			{
+				SkipToNextLine();
+			}
 
 			int32 Count = (int32)(EndOfLine - Start) + 1;
 			String.AppendChars(Start, Count);
@@ -1164,7 +1170,7 @@ namespace CrossCompiler
 		}
 		else if (Tokenizer.MatchString(MATCH_TARGET(TEXT("#pragma"))))
 		{
-			FString Pragma = TEXT("#pragma") + Tokenizer.ReadToEndOfLine();
+			FString Pragma = TEXT("#pragma") + Tokenizer.ReadToEndOfLine(false);
 			Scanner.AddToken(FHlslToken(EHlslToken::Pragma, Pragma), Tokenizer);
 		}
 		else if (Tokenizer.MatchString(MATCH_TARGET(TEXT("#if 0"))))
diff --git a/Engine/Source/Developer/ShaderCompilerCommon/Private/HlslLexer.h b/Engine/Source/Developer/ShaderCompilerCommon/Private/HlslLexer.h
--- a/Engine/Source/Developer/ShaderCompilerCommon/Private/HlslLexer.h
+++ b/Engine/Source/Developer/ShaderCompilerCommon/Private/HlslLexer.h
@@ -249,6 +249,8 @@ namespace CrossCompiler
 		GroupShared,
 		RowMajor,
 		Register,
+		Inline,
+		Typedef,
 
 		Identifier,
 		UnsignedIntegerConstant,
diff --git a/Engine/Source/Developer/ShaderCompilerCommon/Private/HlslParser.cpp b/Engine/Source/Developer/ShaderCompilerCommon/Private/HlslParser.cpp
--- a/Engine/Source/Developer/ShaderCompilerCommon/Private/HlslParser.cpp
+++ b/Engine/Source/Developer/ShaderCompilerCommon/Private/HlslParser.cpp
@@ -266,7 +266,7 @@ namespace CrossCompiler
 			if (Scanner.MatchToken(EHlslToken::Lower))
 			{
 				AST::FTypeSpecifier* ElementTypeSpecifier = nullptr;
-				auto Result = ParseGeneralType(Scanner, ETF_BUILTIN_NUMERIC | ETF_USER_TYPES, SymbolScope, Allocator, &ElementTypeSpecifier);
+				auto Result = ParseGeneralType(Scanner, ETF_BUILTIN_NUMERIC | ETF_USER_TYPES | ETF_UNORM, SymbolScope, Allocator, &ElementTypeSpecifier);
 				if (Result != EParseResult::Matched)
 				{
 					Scanner.SourceError(TEXT("Expected type!"));
@@ -363,7 +363,7 @@ namespace CrossCompiler
 		return EParseResult::NotMatched;
 	}
 
-	EParseResult ParseDeclarationStorageQualifiers(FHlslScanner& Scanner, int32 TypeFlags, int32 DeclarationFlags, bool& bOutPrimitiveFound, AST::FTypeQualifier* Qualifier)
+	EParseResult ParseDeclarationStorageQualifiers(FHlslScanner& Scanner, int32 TypeFlags, int32 DeclarationFlags, bool& bOutPrimitiveFound, FLinearAllocator* Allocator, AST::FTypeQualifier* Qualifier)
 	{
 		bOutPrimitiveFound = false;
 		int32 StaticFound = 0;
@@ -396,6 +396,7 @@ namespace CrossCompiler
 					Token->String == TEXT("triangleadj"))
 				{
 					Scanner.Advance();
+					Qualifier->PrimitiveType = Allocator->Strdup(Token->String);
 					++PrimitiveFound;
 				}
 			}
@@ -603,7 +604,7 @@ namespace CrossCompiler
 		auto OriginalToken = Scanner.GetCurrentTokenIndex();
 		bool bPrimitiveFound = false;
 		auto* FullType = new(Allocator) AST::FFullySpecifiedType(Allocator, Scanner.GetCurrentToken()->SourceInfo);
-		auto ParseResult = ParseDeclarationStorageQualifiers(Scanner, TypeFlags, DeclarationFlags, bPrimitiveFound, &FullType->Qualifier);
+		auto ParseResult = ParseDeclarationStorageQualifiers(Scanner, TypeFlags, DeclarationFlags, bPrimitiveFound, Allocator, &FullType->Qualifier);
 		if (ParseResult == EParseResult::Error)
 		{
 			return ParseResultError();
@@ -1131,9 +1132,22 @@ Done:
 	{
 		const auto* CurrentToken = Parser.Scanner.GetCurrentToken();
 
+		// Inline could be used but will be ignored per the hlsl spec. If found then this HAS to be a function.
+		bool bFoundInline = Parser.Scanner.MatchToken(EHlslToken::Inline);
+
 		AST::FFunction* Function = nullptr;
 		EParseResult Result = ParseFunctionDeclarator(Parser, Allocator, &Function);
-		if (Result == EParseResult::NotMatched || Result == EParseResult::Error)
+		if (Result == EParseResult::NotMatched)
+		{
+			if (bFoundInline)
+			{
+				Parser.Scanner.SourceError(TEXT("Function declaration expected ('inline' keyword found)"));
+				return EParseResult::Error;
+			}
+
+			return EParseResult::NotMatched;
+		}
+		else if( Result == EParseResult::Error)
 		{
 			return Result;
 		}
@@ -1141,6 +1155,7 @@ Done:
 		if (Parser.Scanner.MatchToken(EHlslToken::Semicolon))
 		{
 			// Forward declaration
+			Function->bIsDefinition = true;
 			*OutFunction = Function;
 			return EParseResult::Matched;
 		}
@@ -1741,6 +1756,39 @@ Done:
 		return ParseResultError();
 	}
 
+	EParseResult ParseTypedef(FHlslParser& Parser, FLinearAllocator* Allocator, AST::FNode** OutNode)
+	{
+		if (!Parser.Scanner.MatchToken(EHlslToken::Typedef))
+		{
+			return EParseResult::NotMatched;
+		}
+
+		EParseResult Result = ParseGlobalVariableDeclaration(Parser, Allocator, OutNode);
+		if (Result == EParseResult::Matched)
+		{
+			if (OutNode && *OutNode)
+			{
+				CrossCompiler::AST::FDeclaratorList* DeclList = (*OutNode)->AsDeclaratorList();
+				if (DeclList)
+				{
+					DeclList->bTypedef = true;
+				}
+				else
+				{
+					Parser.Scanner.SourceError(TEXT("Internal error parsing typedef declarator list"));
+					return EParseResult::Error;
+				}
+			}
+			else
+			{
+				Parser.Scanner.SourceError(TEXT("Internal error parsing typedef with null node"));
+				return EParseResult::Error;
+			}
+		}
+
+		return EParseResult::Matched;
+	}
+
 	EParseResult TryTranslationUnit(FHlslParser& Parser, FLinearAllocator* Allocator, AST::FNode** OutNode)
 	{
 		if (MatchPragma(Parser, Allocator, OutNode))
@@ -1786,7 +1834,13 @@ Done:
 			return ParseResultError();
 		}
 		
-		auto Result = ParseFunctionDeclaration(Parser, Allocator, Attributes, OutNode);
+		auto Result = ParseTypedef(Parser, Allocator, OutNode);
+		if (Result == EParseResult::Error || Result == EParseResult::Matched)
+		{
+			return Result;
+		}
+
+		Result = ParseFunctionDeclaration(Parser, Allocator, Attributes, OutNode);
 		if (Result == EParseResult::Error || Result == EParseResult::Matched)
 		{
 			return Result;
