Fix static analysis warning C6011: Dereferencing NULL pointer 'Token'.
#rnx

diff --git a/Engine/Source/Developer/ShaderCompilerCommon/Private/HlslExpressionParser.inl b/Engine/Source/Developer/ShaderCompilerCommon/Private/HlslExpressionParser.inl
--- a/Engine/Source/Developer/ShaderCompilerCommon/Private/HlslExpressionParser.inl
+++ b/Engine/Source/Developer/ShaderCompilerCommon/Private/HlslExpressionParser.inl
@@ -492,7 +492,7 @@ namespace CrossCompiler
 			ExtraQualifiers |= EEQ_PRECISE;
 		}
 
-		if ((TypeFlags & ETF_UNORM) == ETF_UNORM && Token->String.Len() == 5)
+		if ((TypeFlags & ETF_UNORM) == ETF_UNORM && Token && Token->String.Len() == 5)
 		{
 			if (!FCString::Strcmp(*Token->String, TEXT("unorm")))
 			{
