F - Fix static analysis
#rnx

diff --git a/Engine/Source/Developer/ShaderCompilerCommon/Private/HlslParser.cpp b/Engine/Source/Developer/ShaderCompilerCommon/Private/HlslParser.cpp
--- a/Engine/Source/Developer/ShaderCompilerCommon/Private/HlslParser.cpp
+++ b/Engine/Source/Developer/ShaderCompilerCommon/Private/HlslParser.cpp
@@ -1766,7 +1766,7 @@ Done:
 		EParseResult Result = ParseGlobalVariableDeclaration(Parser, Allocator, OutNode);
 		if (Result == EParseResult::Matched)
 		{
-			if (OutNode && *OutNode)
+			if (*OutNode)
 			{
 				CrossCompiler::AST::FDeclaratorList* DeclList = (*OutNode)->AsDeclaratorList();
 				if (DeclList)
