Merging

//UE4/Dev-Rendering/Engine/Source/Developer/ShaderCompilerCommon/Private/ShaderCompilerCommon.cpp

to //UE4/Main/Engine/Source/Developer/ShaderCompilerCommon/Private/ShaderCompilerCommon.cpp

#rnx

diff --git a/Engine/Source/Developer/ShaderCompilerCommon/Private/ShaderCompilerCommon.cpp b/Engine/Source/Developer/ShaderCompilerCommon/Private/ShaderCompilerCommon.cpp
--- a/Engine/Source/Developer/ShaderCompilerCommon/Private/ShaderCompilerCommon.cpp
+++ b/Engine/Source/Developer/ShaderCompilerCommon/Private/ShaderCompilerCommon.cpp
@@ -951,7 +951,14 @@ void FShaderParameterParser::ValidateShaderParameterTypes(
 		const bool bShouldBeInt = Member.ExpectedShaderType.StartsWith(TEXT("int"));
 		const bool bShouldBeUint = Member.ExpectedShaderType.StartsWith(TEXT("uint"));
 
+		// Match parsed type with expected shader type
 		bool bIsTypeCorrect = ParsedParameter.Type == Member.ExpectedShaderType;
+		
+		// Accept half-precision floats when single-precision was requested
+		if (!bIsTypeCorrect && ParsedParameter.Type.StartsWith(TEXT("half")) && Member.ExpectedShaderType.StartsWith(TEXT("float")))
+		{
+			bIsTypeCorrect = (FCString::Strcmp(*ParsedParameter.Type + 4, *Member.ExpectedShaderType + 5) == 0);
+		}
 
 		// Allow silent casting between signed and unsigned on shader bindings.
 		if (!bIsTypeCorrect && (bShouldBeInt || bShouldBeUint))
