UE4.25 - Fix bad module dependencies
#rnx

diff --git a/Engine/Source/Developer/ShaderCompilerCommon/Public/ShaderCompilerCommon.h b/Engine/Source/Developer/ShaderCompilerCommon/Public/ShaderCompilerCommon.h
--- a/Engine/Source/Developer/ShaderCompilerCommon/Public/ShaderCompilerCommon.h
+++ b/Engine/Source/Developer/ShaderCompilerCommon/Public/ShaderCompilerCommon.h
@@ -357,21 +357,5 @@ namespace CrossCompiler
 
 }
 
-enum class ESCWErrorCode
-{
-	NotSet = -1,
-	Success,
-	GeneralCrash,
-	BadShaderFormatVersion,
-	BadInputVersion,
-	BadSingleJobHeader,
-	BadPipelineJobHeader,
-	CantDeleteInputFile,
-	CantSaveOutputFile,
-	NoTargetShaderFormatsFound,
-	CantCompileForSpecificFormat,
-	CrashInsidePlatformCompiler,
-};
-
 // Error code for SCW to help track down crashes
 extern SHADERCOMPILERCOMMON_API ESCWErrorCode GSCWErrorCode;
diff --git a/Engine/Source/Runtime/Engine/Engine.Build.cs b/Engine/Source/Runtime/Engine/Engine.Build.cs
--- a/Engine/Source/Runtime/Engine/Engine.Build.cs
+++ b/Engine/Source/Runtime/Engine/Engine.Build.cs
@@ -36,8 +36,7 @@ public class Engine : ModuleRules
 				"MovieSceneCapture",
 				"AutomationWorker",
 				"MovieSceneCapture",
-				"DesktopPlatform",
-				"ShaderCompilerCommon"
+				"DesktopPlatform"
 			}
 		);
 
diff --git a/Engine/Source/Runtime/Engine/Private/ShaderCompiler/ShaderCompiler.cpp b/Engine/Source/Runtime/Engine/Private/ShaderCompiler/ShaderCompiler.cpp
--- a/Engine/Source/Runtime/Engine/Private/ShaderCompiler/ShaderCompiler.cpp
+++ b/Engine/Source/Runtime/Engine/Private/ShaderCompiler/ShaderCompiler.cpp
@@ -48,7 +48,7 @@
 #include "ShaderParameterMetadata.h"
 #include "ProfilingDebugging/DiagnosticTable.h"
 #include "ProfilingDebugging/LoadTimeTracker.h"
-#include "ShaderCompilerCommon.h"
+#include "ShaderCore.h"
 
 #define LOCTEXT_NAMESPACE "ShaderCompiler"
 
diff --git a/Engine/Source/Runtime/RenderCore/Public/ShaderCore.h b/Engine/Source/Runtime/RenderCore/Public/ShaderCore.h
--- a/Engine/Source/Runtime/RenderCore/Public/ShaderCore.h
+++ b/Engine/Source/Runtime/RenderCore/Public/ShaderCore.h
@@ -1288,4 +1288,20 @@ void ApplyGlobalUniformBuffers(
 			}
 		}
 	}
-}
\ No newline at end of file
+}
+
+enum class ESCWErrorCode
+{
+	NotSet = -1,
+	Success,
+	GeneralCrash,
+	BadShaderFormatVersion,
+	BadInputVersion,
+	BadSingleJobHeader,
+	BadPipelineJobHeader,
+	CantDeleteInputFile,
+	CantSaveOutputFile,
+	NoTargetShaderFormatsFound,
+	CantCompileForSpecificFormat,
+	CrashInsidePlatformCompiler,
+};
