Fixed 0 hash in shader compiler for uniform buffers without resources.

diff --git a/Engine/Source/Developer/ShaderCompilerCommon/Private/ShaderCompilerCommon.cpp b/Engine/Source/Developer/ShaderCompilerCommon/Private/ShaderCompilerCommon.cpp
--- a/Engine/Source/Developer/ShaderCompilerCommon/Private/ShaderCompilerCommon.cpp
+++ b/Engine/Source/Developer/ShaderCompilerCommon/Private/ShaderCompilerCommon.cpp
@@ -153,22 +153,26 @@ bool BuildResourceTableMapping(
 				return false;
 			}
 		}
+	}
 
-		// We have to do this separately from the resource table member check above. We want to include the hash even
-		// if the uniform buffer does not have any actual members used, because it will still be in the parameter map
-		// and certain platforms (like DX12) will pessimise and require them.
-		{
-			uint16 UniformBufferIndex = INDEX_NONE;
-			uint16 UBBaseIndex, UBSize;
+	// Emit hashes for all uniform buffers in the parameter map. We need to include the ones without resources as well
+	// (i.e. just constants), since the global uniform buffer bindings rely on valid hashes.
+	for (const auto& KeyValue : ParameterMap.GetParameterMap())
+	{
+		const FString& UniformBufferName = KeyValue.Key;
+		const FParameterAllocation& UniformBufferParameter = KeyValue.Value;
 
-			if (ParameterMap.FindParameterAllocation(*Entry.UniformBufferName, UniformBufferIndex, UBBaseIndex, UBSize))
+		if (UniformBufferParameter.Type == EShaderParameterType::UniformBuffer)
+		{
+			if (OutSRT.ResourceTableLayoutHashes.Num() <= UniformBufferParameter.BufferIndex)
 			{
-				while (OutSRT.ResourceTableLayoutHashes.Num() <= UniformBufferIndex)
-				{
-					OutSRT.ResourceTableLayoutHashes.Add(0);
-				}
+				OutSRT.ResourceTableLayoutHashes.SetNumZeroed(UniformBufferParameter.BufferIndex + 1);
+			}
 
-				OutSRT.ResourceTableLayoutHashes[UniformBufferIndex] = ResourceTableLayoutHashes.FindChecked(Entry.UniformBufferName);
+			// Data-driven uniform buffers will not have registered this information.
+			if (const uint32* HashPtr = ResourceTableLayoutHashes.Find(UniformBufferName))
+			{
+				OutSRT.ResourceTableLayoutHashes[UniformBufferParameter.BufferIndex] = *HashPtr;
 			}
 		}
 	}
