DXC Update: SCW backends.

#rnx

diff --git a/Engine/Shaders/Public/Platform/Metal/MetalCommon.ush b/Engine/Shaders/Public/Platform/Metal/MetalCommon.ush
--- a/Engine/Shaders/Public/Platform/Metal/MetalCommon.ush
+++ b/Engine/Shaders/Public/Platform/Metal/MetalCommon.ush
@@ -5,7 +5,7 @@
 =============================================================================*/
 
 // Update this GUID to improve shader recompilation for Metal only shaders
-// GUID = 75802E83-9450-4B4A-A2A6-93B080C7C8DC
+// GUID = 7742EA46-5497-4ABC-897F-132501941942
 
 #pragma once
 
diff --git a/Engine/Shaders/Public/Platform/Vulkan/VulkanCommon.ush b/Engine/Shaders/Public/Platform/Vulkan/VulkanCommon.ush
--- a/Engine/Shaders/Public/Platform/Vulkan/VulkanCommon.ush
+++ b/Engine/Shaders/Public/Platform/Vulkan/VulkanCommon.ush
@@ -5,6 +5,6 @@
 =============================================================================*/
 
 // Update this GUID to improve shader recompilation for Vulkan only shaders
-// GUID = 0D1873CC-3A3B-49C2-9D78-6317A87C2857
+// GUID = 0B7690A1-083B-45B3-8A4D-485045A00CFB
 
 #pragma once
diff --git a/Engine/Source/Developer/Apple/MetalShaderFormat/Private/MetalDerivedData.cpp b/Engine/Source/Developer/Apple/MetalShaderFormat/Private/MetalDerivedData.cpp
--- a/Engine/Source/Developer/Apple/MetalShaderFormat/Private/MetalDerivedData.cpp
+++ b/Engine/Source/Developer/Apple/MetalShaderFormat/Private/MetalDerivedData.cpp
@@ -1674,13 +1674,13 @@ bool FMetalShaderOutputCooker::Build(TArray<uint8>& OutData)
 			case EMetalGPUSemanticsTBDRDesktop:
 				TargetDesc.Language = CrossCompiler::EShaderConductorLanguage::Metal_iOS;
 				TargetDesc.CompileFlags.SetDefine(TEXT("ios_support_base_vertex_instance"), 1);
-				TargetDesc.CompileFlags.SetDefine(TEXT("ios_use_framebuffer_fetch_subpasses"), 1);
+				TargetDesc.CompileFlags.SetDefine(TEXT("use_framebuffer_fetch_subpasses"), 1);
 				TargetDesc.CompileFlags.SetDefine(TEXT("emulate_cube_array"), 1);
 				break;
 			case EMetalGPUSemanticsMobile:
 			default:
 				TargetDesc.Language = CrossCompiler::EShaderConductorLanguage::Metal_iOS;
-				TargetDesc.CompileFlags.SetDefine(TEXT("ios_use_framebuffer_fetch_subpasses"), 1);
+				TargetDesc.CompileFlags.SetDefine(TEXT("use_framebuffer_fetch_subpasses"), 1);
 				TargetDesc.CompileFlags.SetDefine(TEXT("emulate_cube_array"), 1);
 				break;
 			}
diff --git a/Engine/Source/Developer/ShaderCompilerCommon/Private/ShaderCompilerCommon.cpp b/Engine/Source/Developer/ShaderCompilerCommon/Private/ShaderCompilerCommon.cpp
--- a/Engine/Source/Developer/ShaderCompilerCommon/Private/ShaderCompilerCommon.cpp
+++ b/Engine/Source/Developer/ShaderCompilerCommon/Private/ShaderCompilerCommon.cpp
@@ -2498,9 +2498,10 @@ namespace CrossCompiler
 		OutResultDesc = ShaderConductor::Compiler::Rewrite(InSourceDesc, InOptions);
 	}
 
-	static ShaderConductor::Compiler::ResultDesc ScRewriteWrapper(
+	static bool ScRewriteWrapper(
 		const ShaderConductor::Compiler::SourceDesc& InSourceDesc,
 		const ShaderConductor::Compiler::Options& InOptions,
+		ShaderConductor::Compiler::ResultDesc& OutResultDesc,
 		bool& bOutException)
 	{
 		bOutException = false;
@@ -2508,26 +2509,35 @@ namespace CrossCompiler
 		__try
 #endif
 		{
-			ShaderConductor::Compiler::ResultDesc Result;
-			InnerScRewriteWrapper(InSourceDesc, InOptions, Result);
-			return Result;
+			InnerScRewriteWrapper(InSourceDesc, InOptions, OutResultDesc);
+			return true;
 		}
 #if !PLATFORM_SEH_EXCEPTIONS_DISABLED
 		__except (EXCEPTION_EXECUTE_HANDLER)
 		{
 			GSCWErrorCode = ESCWErrorCode::CrashInsidePlatformCompiler;
-			ShaderConductor::Compiler::ResultDesc ResultDesc;
-			FMemory::Memzero(ResultDesc);
+			FMemory::Memzero(OutResultDesc);
 			bOutException = true;
-			return ResultDesc;
+			return false;
 		}
 #endif
 	}
 
-	static ShaderConductor::Compiler::ResultDesc ScCompileWrapper(
+	// Inner wrapper function is required here because '__try'-statement cannot be used with function that requires object unwinding
+	static void InnerScCompileWrapper(
 		const ShaderConductor::Compiler::SourceDesc& InSourceDesc,
 		const ShaderConductor::Compiler::Options& InOptions,
 		const ShaderConductor::Compiler::TargetDesc& InTargetDesc,
+		ShaderConductor::Compiler::ResultDesc& OutResultDesc)
+	{
+		OutResultDesc = ShaderConductor::Compiler::Compile(InSourceDesc, InOptions, InTargetDesc);
+	}
+
+	static bool ScCompileWrapper(
+		const ShaderConductor::Compiler::SourceDesc& InSourceDesc,
+		const ShaderConductor::Compiler::Options& InOptions,
+		const ShaderConductor::Compiler::TargetDesc& InTargetDesc,
+		ShaderConductor::Compiler::ResultDesc& OutResultDesc,
 		bool& bOutException)
 	{
 		bOutException = false;
@@ -2535,25 +2545,35 @@ namespace CrossCompiler
 		__try
 #endif
 		{
-			ShaderConductor::Compiler::ResultDesc Result = ShaderConductor::Compiler::Compile(InSourceDesc, InOptions, InTargetDesc);
-			return Result;
+			InnerScCompileWrapper(InSourceDesc, InOptions, InTargetDesc, OutResultDesc);
+			return true;
 		}
 #if !PLATFORM_SEH_EXCEPTIONS_DISABLED
 		__except (EXCEPTION_EXECUTE_HANDLER)
 		{
 			GSCWErrorCode = ESCWErrorCode::CrashInsidePlatformCompiler;
-			ShaderConductor::Compiler::ResultDesc ResultDesc;
-			FMemory::Memzero(ResultDesc);
+			FMemory::Memzero(OutResultDesc);
 			bOutException = true;
-			return ResultDesc;
+			return false;
 		}
 #endif
 	}
 
-	static ShaderConductor::Compiler::ResultDesc ScConvertBinaryWrapper(
+	// Inner wrapper function is required here because '__try'-statement cannot be used with function that requires object unwinding
+	static void InnerScCompileWrapper(
+		const ShaderConductor::Compiler::ResultDesc& InBinaryDesc,
+		const ShaderConductor::Compiler::SourceDesc& InSourceDesc,
+		const ShaderConductor::Compiler::TargetDesc& InTargetDesc,
+		ShaderConductor::Compiler::ResultDesc& OutResultDesc)
+	{
+		OutResultDesc = ShaderConductor::Compiler::ConvertBinary(InBinaryDesc, InSourceDesc, InTargetDesc);
+	}
+
+	static bool ScConvertBinaryWrapper(
 		const ShaderConductor::Compiler::ResultDesc& InBinaryDesc,
 		const ShaderConductor::Compiler::SourceDesc& InSourceDesc,
 		const ShaderConductor::Compiler::TargetDesc& InTargetDesc,
+		ShaderConductor::Compiler::ResultDesc& OutResultDesc,
 		bool& bOutException)
 	{
 		bOutException = false;
@@ -2561,16 +2581,16 @@ namespace CrossCompiler
 		__try
 #endif
 		{
-			return ShaderConductor::Compiler::ConvertBinary(InBinaryDesc, InSourceDesc, InTargetDesc);
+			InnerScCompileWrapper(InBinaryDesc, InSourceDesc, InTargetDesc, OutResultDesc);
+			return true;
 		}
 #if !PLATFORM_SEH_EXCEPTIONS_DISABLED
 		__except (EXCEPTION_EXECUTE_HANDLER)
 		{
 			GSCWErrorCode = ESCWErrorCode::CrashInsidePlatformCompiler;
-			ShaderConductor::Compiler::ResultDesc ResultDesc;
-			FMemory::Memzero(ResultDesc);
+			FMemory::Memzero(OutResultDesc);
 			bOutException = true;
-			return ResultDesc;
+			return false;
 		}
 #endif
 	}
@@ -2701,7 +2721,6 @@ namespace CrossCompiler
 	static void ConvertScTargetDescLanguageGlslFamily(const FShaderConductorTarget& InTarget, ShaderConductor::Compiler::TargetDesc& OutTargetDesc)
 	{
 		OutTargetDesc.language = (InTarget.Language == EShaderConductorLanguage::Glsl ? ShaderConductor::ShadingLanguage::Glsl : ShaderConductor::ShadingLanguage::Essl);
-		OutTargetDesc.platform = "";
 		OutTargetDesc.version = GetGlslFamilyVersionString(InTarget.Version);
 		checkf(OutTargetDesc.version, TEXT("Unsupported target shader version for GLSL family: %d"), InTarget.Version);
 	}
@@ -2721,8 +2740,7 @@ namespace CrossCompiler
 
 	static void ConvertScTargetDescLanguageMetalFamily(const FShaderConductorTarget& InTarget, ShaderConductor::Compiler::TargetDesc& OutTargetDesc)
 	{
-		OutTargetDesc.language = ShaderConductor::ShadingLanguage::Msl;
-		OutTargetDesc.platform = (InTarget.Language == EShaderConductorLanguage::Metal_macOS ? "macOS" : "iOS");
+		OutTargetDesc.language = (InTarget.Language == EShaderConductorLanguage::Metal_macOS ? ShaderConductor::ShadingLanguage::Msl_macOS : ShaderConductor::ShadingLanguage::Msl_iOS);
 		OutTargetDesc.version = GetMetalFamilyVersionString(InTarget.Version);
 		checkf(OutTargetDesc.version, TEXT("Unsupported target shader version for Metal family: %d"), InTarget.Version);
 	}
@@ -2767,7 +2785,7 @@ namespace CrossCompiler
 		// Wrap input function into lambda to convert to ShaderConductor interface
 		if (InTarget.VariableTypeRenameCallback)
 		{
-			OutTargetDesc.variableTypeRenameCallback = [InnerCallback = InTarget.VariableTypeRenameCallback](const char* VariableName, const char* TypeName) -> ShaderConductor::Blob*
+			OutTargetDesc.variableTypeRenameCallback = [InnerCallback = InTarget.VariableTypeRenameCallback](const char* VariableName, const char* TypeName) -> ShaderConductor::Blob
 			{
 				// Forward callback to public interface callback
 				FString RenamedTypeName;
@@ -2776,10 +2794,10 @@ namespace CrossCompiler
 					if (!RenamedTypeName.IsEmpty())
 					{
 						// Convert renamed type name from FString to ShaderConductor::Blob
-						return ShaderConductor::CreateBlob(TCHAR_TO_ANSI(*RenamedTypeName), RenamedTypeName.Len() + 1);
+						return ShaderConductor::Blob(TCHAR_TO_ANSI(*RenamedTypeName), RenamedTypeName.Len() + 1);
 					}
 				}
-				return nullptr;
+				return ShaderConductor::Blob{};
 			};
 		}
 	}
@@ -2792,7 +2810,6 @@ namespace CrossCompiler
 		OutOptions.enableDebugInfo = InOptions.bEnableDebugInfo;
 		OutOptions.disableOptimizations = InOptions.bDisableOptimizations;
 		OutOptions.enableFMAPass = InOptions.bEnableFMAPass;
-		OutOptions.globalsAsPushConstants = InOptions.bGlobalsAsPushConstants;
 		OutOptions.shaderModel = ToShaderConductorShaderModel(InOptions.TargetProfile);
 	}
 
@@ -2853,46 +2870,6 @@ namespace CrossCompiler
 		}
 	}
 
-	// Implements the ShaderConductor::Blob interface with a weak reference to a block of data.
-	class FShaderConductorWeakRefBlob : public ShaderConductor::Blob
-	{
-	public:
-		FShaderConductorWeakRefBlob(const FShaderConductorWeakRefBlob&) = delete;
-		FShaderConductorWeakRefBlob& operator = (const FShaderConductorWeakRefBlob&) = delete;
-
-		FShaderConductorWeakRefBlob()
-			: DataPtr(nullptr)
-			, DataSize(0)
-		{
-		}
-
-		FShaderConductorWeakRefBlob(const void* InData, uint32 InSize)
-			: DataPtr(InData)
-			, DataSize(InSize)
-		{
-		}
-
-		FShaderConductorWeakRefBlob(FShaderConductorWeakRefBlob&& Rhs)
-			: DataPtr(Rhs.DataPtr)
-			, DataSize(Rhs.DataSize)
-		{
-			Rhs.DataPtr = nullptr;
-			Rhs.DataSize = 0;
-		}
-
-		virtual const void* Data() const override
-		{
-			return DataPtr;
-		}
-		virtual uint32_t Size() const override
-		{
-			return DataSize;
-		}
-	private:
-		const void* DataPtr;
-		uint32_t DataSize;
-	};
-
 	FShaderConductorContext::FShaderConductorContext()
 		: Intermediates(new FShaderConductorIntermediates())
 	{
@@ -2969,13 +2946,13 @@ namespace CrossCompiler
 		// Rewrite HLSL with wrapper function to catch exceptions from ShaderConductor
 		bool bSucceeded = false;
 		bool bException = false;
-		ShaderConductor::Compiler::ResultDesc ResultDesc = ScRewriteWrapper(ScSourceDesc, ScOptions, bException);
-		ShaderConductor::Blob* RewriteBlob = ResultDesc.target;
+		ShaderConductor::Compiler::ResultDesc ResultDesc;
+		ScRewriteWrapper(ScSourceDesc, ScOptions, ResultDesc, bException);
 
-		if (!ResultDesc.hasError && !bException && RewriteBlob != nullptr)
+		if (!ResultDesc.hasError && !bException && ResultDesc.target.Size() > 0)
 		{
 			// Copy rewritten HLSL code into intermediate source code.
-			ConvertByteArrayToAnsiString(reinterpret_cast<const ANSICHAR*>(RewriteBlob->Data()), RewriteBlob->Size(), Intermediates->ShaderSource);
+			ConvertByteArrayToAnsiString(reinterpret_cast<const ANSICHAR*>(ResultDesc.target.Data()), ResultDesc.target.Size(), Intermediates->ShaderSource);
 
 			// If output source is specified, also convert to TCHAR string
 			if (OutSource != nullptr)
@@ -2994,17 +2971,7 @@ namespace CrossCompiler
 		}
 
 		// Append compile error and warning to output reports
-		if (ShaderConductor::Blob* ErrorBlob = ResultDesc.errorWarningMsg)
-		{
-			ConvertScCompileErrors(*ErrorBlob, Errors);
-			ShaderConductor::DestroyBlob(ErrorBlob);
-		}
-
-		// Clean up intermediate buffers
-		if (RewriteBlob)
-		{
-			ShaderConductor::DestroyBlob(RewriteBlob);
-		}
+		ConvertScCompileErrors(ResultDesc.errorWarningMsg, Errors);
 
 		return bSucceeded;
 	}
@@ -3025,12 +2992,13 @@ namespace CrossCompiler
 		// Compile HLSL source code to SPIR-V
 		bool bSucceeded = false;
 		bool bException = false;
-		ShaderConductor::Compiler::ResultDesc ResultDesc = ScCompileWrapper(ScSourceDesc, ScOptions, ScTargetDesc, bException);
+		ShaderConductor::Compiler::ResultDesc ResultDesc;
+		ScCompileWrapper(ScSourceDesc, ScOptions, ScTargetDesc, ResultDesc, bException);
 
-		if (!ResultDesc.hasError && !bException && ResultDesc.target != nullptr)
+		if (!ResultDesc.hasError && !bException && ResultDesc.target.Size() > 0)
 		{
 			// Copy result blob into output SPIR-V module
-			OutSpirv = TArray<uint32>(reinterpret_cast<const uint32*>(ResultDesc.target->Data()), ResultDesc.target->Size() / 4);
+			OutSpirv = TArray<uint32>(reinterpret_cast<const uint32*>(ResultDesc.target.Data()), ResultDesc.target.Size() / 4);
 			bSucceeded = true;
 		}
 		else
@@ -3043,17 +3011,7 @@ namespace CrossCompiler
 		}
 
 		// Append compile error and warning to output reports
-		if (ShaderConductor::Blob* ErrorBlob = ResultDesc.errorWarningMsg)
-		{
-			ConvertScCompileErrors(*ErrorBlob, Errors);
-			ShaderConductor::DestroyBlob(ErrorBlob);
-		}
-
-		// Clean up intermediate buffers
-		if (ResultDesc.target)
-		{
-			ShaderConductor::DestroyBlob(ResultDesc.target);
-		}
+		ConvertScCompileErrors(ResultDesc.errorWarningMsg, Errors);
 
 		return bSucceeded;
 	}
@@ -3099,24 +3057,21 @@ namespace CrossCompiler
 		ShaderConductor::Compiler::Options ScOptions;
 		ConvertScOptions(Options, ScOptions);
 
-		// Create temporary weak reference to SPIR-V provided in ShaderConductor::Blob interface.
-		// Avoid copy, so don't use ShaderConductor::CreateBlob().
-		FShaderConductorWeakRefBlob SpirvBlob(InSpirv, InSpirvByteSize);
 		ShaderConductor::Compiler::ResultDesc ScBinaryDesc;
-		ScBinaryDesc.target = &SpirvBlob;
+		ScBinaryDesc.target.Reset(InSpirv, InSpirvByteSize);
 		ScBinaryDesc.isText = false;
-		ScBinaryDesc.errorWarningMsg = nullptr;
 		ScBinaryDesc.hasError = false;
 
 		// Convert the input SPIR-V into Metal high level source
 		bool bSucceeded = false;
 		bool bException = false;
-		ShaderConductor::Compiler::ResultDesc ResultDesc = ScConvertBinaryWrapper(ScBinaryDesc, ScSourceDesc, ScTargetDesc, bException);
+		ShaderConductor::Compiler::ResultDesc ResultDesc;
+		ScConvertBinaryWrapper(ScBinaryDesc, ScSourceDesc, ScTargetDesc, ResultDesc, bException);
 
-		if (!ResultDesc.hasError && !bException && ResultDesc.target != nullptr)
+		if (!ResultDesc.hasError && !bException && ResultDesc.target.Size() > 0)
 		{
 			// Copy result blob into output SPIR-V module
-			OutputCallback(ResultDesc.target->Data(), ResultDesc.target->Size());
+			OutputCallback(ResultDesc.target.Data(), ResultDesc.target.Size());
 			bSucceeded = true;
 		}
 		else
@@ -3129,20 +3084,13 @@ namespace CrossCompiler
 		}
 
 		// Append compile error and warning to output reports
-		if (ShaderConductor::Blob* ErrorBlob = ResultDesc.errorWarningMsg)
+		if (ResultDesc.errorWarningMsg.Size() > 0)
 		{
 			FString ErrorString;
-			if (ConvertScBlobToFString(ErrorBlob, ErrorString))
+			if (ConvertScBlobToFString(&ResultDesc.errorWarningMsg, ErrorString))
 			{
 				Errors.Add(*ErrorString);
 			}
-			ShaderConductor::DestroyBlob(ResultDesc.errorWarningMsg);
-		}
-
-		// Clean up intermediate buffers
-		if (ResultDesc.target)
-		{
-			ShaderConductor::DestroyBlob(ResultDesc.target);
 		}
 
 		return bSucceeded;
diff --git a/Engine/Source/Developer/ShaderCompilerCommon/Public/ShaderCompilerCommon.h b/Engine/Source/Developer/ShaderCompilerCommon/Public/ShaderCompilerCommon.h
--- a/Engine/Source/Developer/ShaderCompilerCommon/Public/ShaderCompilerCommon.h
+++ b/Engine/Source/Developer/ShaderCompilerCommon/Public/ShaderCompilerCommon.h
@@ -387,8 +387,8 @@ namespace CrossCompiler
 		/** Removes unused global variables and resources. This can only be used in the HLSL rewrite pass, i.e. 'RewriteHlslSource'. */
 		bool bRemoveUnusedGlobals = false;
 
-		/** Experimental: Decide how a matrix get packed. */
-		bool bPackMatricesInRowMajor = false;
+		/** Experimental: Decide how a matrix get packed. Default in HLSL is row-major. This will be inverted in the SPIR-V backend to match SPIR-V's column-major default. */
+		bool bPackMatricesInRowMajor = true;
 
 		/** Enable 16-bit types, such as half, uint16_t. Requires shader model 6.2+. */
 		bool bEnable16bitTypes = false;
@@ -402,9 +402,6 @@ namespace CrossCompiler
 		/** Enable a pass that converts floating point MUL+ADD pairs into FMAs to avoid re-association. */
 		bool bEnableFMAPass = false;
 
-		/** Cross compile global variables as push constants (for Vulkan backend). */
-		bool bGlobalsAsPushConstants = false;
-
 		/** Target shader profile. By default HCT_FeatureLevelSM5. */
 		EHlslCompileTarget TargetProfile = HCT_FeatureLevelSM5;
 	};
diff --git a/Engine/Source/Developer/ShaderFormatOpenGL/Private/OpenGLShaderCompiler.cpp b/Engine/Source/Developer/ShaderFormatOpenGL/Private/OpenGLShaderCompiler.cpp
--- a/Engine/Source/Developer/ShaderFormatOpenGL/Private/OpenGLShaderCompiler.cpp
+++ b/Engine/Source/Developer/ShaderFormatOpenGL/Private/OpenGLShaderCompiler.cpp
@@ -1244,7 +1244,6 @@ static bool CompileToGlslWithShaderConductor(
 
 	// Initialize compilation options for ShaderConductor
 	CrossCompiler::FShaderConductorOptions Options;
-	Options.bGlobalsAsPushConstants = true;
 
 	// Convert input strings from FString to ANSI strings
 	std::string SourceData(TCHAR_TO_UTF8(*PreprocessedShader));
diff --git a/Engine/Source/Developer/VulkanShaderFormat/Private/VulkanShaderCompiler.cpp b/Engine/Source/Developer/VulkanShaderFormat/Private/VulkanShaderCompiler.cpp
--- a/Engine/Source/Developer/VulkanShaderFormat/Private/VulkanShaderCompiler.cpp
+++ b/Engine/Source/Developer/VulkanShaderFormat/Private/VulkanShaderCompiler.cpp
@@ -2258,7 +2258,7 @@ static const TCHAR* SpvBuiltinToString(const SpvBuiltIn BuiltIn)
 	return nullptr;
 }
 
-static FString ConvertMetaDataSemantic(const FString& InSemantic, const SpvBuiltIn BuiltIn, bool bIsInput)
+static FString ConvertAttributeToMetaDataSemantic(const ANSICHAR* AttributeName, const SpvBuiltIn BuiltIn, bool bIsInput)
 {
 	if (const TCHAR* BuiltInName = SpvBuiltinToString(BuiltIn))
 	{
@@ -2266,6 +2266,8 @@ static FString ConvertMetaDataSemantic(const FString& InSemantic, const SpvBuilt
 	}
 	else
 	{
+		check(AttributeName != nullptr && *AttributeName != '\0');
+		FString InSemantic = ANSI_TO_TCHAR(AttributeName);
 		FString OutSemantic = (bIsInput ? TEXT("in_") : TEXT("out_"));
 
 		if (InSemantic.StartsWith(TEXT("SV_")))
@@ -2304,9 +2306,8 @@ static void BuildShaderInterfaceVariableMetaData(const SpvReflectInterfaceVariab
 		return;
 	}
 
-	check(Attribute.semantic != nullptr);
 	const FString TypeSpecifier = ConvertMetaDataTypeSpecifier(*Attribute.type_description);
-	FString Semantic = ConvertMetaDataSemantic(ANSI_TO_TCHAR(Attribute.semantic), Attribute.built_in, bIsInput);
+	FString Semantic = ConvertAttributeToMetaDataSemantic(Attribute.semantic, Attribute.built_in, bIsInput);
 
 	if (Attribute.array.dims_count > 0)
 	{
diff --git a/Engine/Source/Developer/Windows/ShaderFormatD3D/Private/D3DShaderCompilerDXC.cpp b/Engine/Source/Developer/Windows/ShaderFormatD3D/Private/D3DShaderCompilerDXC.cpp
--- a/Engine/Source/Developer/Windows/ShaderFormatD3D/Private/D3DShaderCompilerDXC.cpp
+++ b/Engine/Source/Developer/Windows/ShaderFormatD3D/Private/D3DShaderCompilerDXC.cpp
@@ -460,14 +460,19 @@ static ShaderConductor::ShaderStage ToDXCShaderStage(EShaderFrequency Frequency)
 	}
 }
 
-static void InnerDXCRewriteWrapper(const ShaderConductor::Compiler::SourceDesc& InDesc,
-	const ShaderConductor::Compiler::Options& InOptions, ShaderConductor::Compiler::ResultDesc& ResultDesc)
+// Inner wrapper function is required here because '__try'-statement cannot be used with function that requires object unwinding
+static void InnerScRewriteWrapper(
+	const ShaderConductor::Compiler::SourceDesc& InDesc,
+	const ShaderConductor::Compiler::Options& InOptions,
+	ShaderConductor::Compiler::ResultDesc& OutResultDesc)
 {
-	ResultDesc = ShaderConductor::Compiler::Rewrite(InDesc, InOptions);
+	OutResultDesc = ShaderConductor::Compiler::Rewrite(InDesc, InOptions);
 }
 
-static ShaderConductor::Compiler::ResultDesc DXCRewriteWrapper(const ShaderConductor::Compiler::SourceDesc& InDesc,
+static bool DXCRewriteWrapper(
+	const ShaderConductor::Compiler::SourceDesc& InDesc,
 	const ShaderConductor::Compiler::Options& InOptions,
+	ShaderConductor::Compiler::ResultDesc& OutResultDesc,
 	bool& bOutException)
 {
 	bOutException = false;
@@ -475,18 +480,16 @@ static ShaderConductor::Compiler::ResultDesc DXCRewriteWrapper(const ShaderCondu
 	__try
 #endif
 	{
-		ShaderConductor::Compiler::ResultDesc ResultDesc;
-		InnerDXCRewriteWrapper(InDesc, InOptions, ResultDesc);
-		return ResultDesc;
+		InnerScRewriteWrapper(InDesc, InOptions, OutResultDesc);
+		return true;
 	}
 #if !PLATFORM_SEH_EXCEPTIONS_DISABLED
 	__except (EXCEPTION_EXECUTE_HANDLER)
 	{
 		GSCWErrorCode = ESCWErrorCode::CrashInsidePlatformCompiler;
-		ShaderConductor::Compiler::ResultDesc ResultDesc;
-		FMemory::Memzero(ResultDesc);
+		FMemory::Memzero(OutResultDesc);
 		bOutException = true;
-		return ResultDesc;
+		return false;
 	}
 #endif
 }
@@ -538,12 +541,10 @@ static bool RewriteUsingSC(FString& PreprocessedShaderSource, const FShaderCompi
 		TargetDesc.language = ShaderConductor::ShadingLanguage::Dxil;
 
 		// Rewrite HLSL source to remove unused global variables (DXC retains them when compiling)
-		ShaderConductor::Blob* RewriteBlob = nullptr;
-
-		// Rewrite HLSL
 		Options.removeUnusedGlobals = true;
 		bool bException = false;
-		ShaderConductor::Compiler::ResultDesc RewriteResultDesc = DXCRewriteWrapper(SourceDesc, Options, bException);
+		ShaderConductor::Compiler::ResultDesc RewriteResultDesc;
+		DXCRewriteWrapper(SourceDesc, Options, RewriteResultDesc, bException);
 		Options.removeUnusedGlobals = false;
 		if (RewriteResultDesc.hasError || bException)
 		{
@@ -552,25 +553,21 @@ static bool RewriteUsingSC(FString& PreprocessedShaderSource, const FShaderCompi
 				Output.Errors.Add(TEXT("ShaderConductor exception during rewrite"));
 			}
 			// Append compile error to output reports
-			if (ShaderConductor::Blob* ErrorBlob = RewriteResultDesc.errorWarningMsg)
+			if (RewriteResultDesc.errorWarningMsg.Size() > 0)
 			{
-				FUTF8ToTCHAR UTF8Converter(reinterpret_cast<const ANSICHAR*>(ErrorBlob->Data()), ErrorBlob->Size());
-				const FString ErrorString(ErrorBlob->Size(), UTF8Converter.Get());
+				FUTF8ToTCHAR UTF8Converter(reinterpret_cast<const ANSICHAR*>(RewriteResultDesc.errorWarningMsg.Data()), RewriteResultDesc.errorWarningMsg.Size());
+				const FString ErrorString(RewriteResultDesc.errorWarningMsg.Size(), UTF8Converter.Get());
 				Output.Errors.Add(*ErrorString);
-
-				ShaderConductor::DestroyBlob(RewriteResultDesc.errorWarningMsg);
-				RewriteResultDesc.errorWarningMsg = nullptr;
+				RewriteResultDesc.errorWarningMsg.Reset();
 				bResult = false;
 			}
 		}
 		else
 		{
 			// Copy rewritten HLSL code into new source data string
-			RewriteBlob = RewriteResultDesc.target;
-
 			CStrSourceData.clear();
-			CStrSourceData.resize(RewriteBlob->Size());
-			FCStringAnsi::Strncpy(&CStrSourceData[0], static_cast<const char*>(RewriteBlob->Data()), RewriteBlob->Size());
+			CStrSourceData.resize(RewriteResultDesc.target.Size());
+			FCStringAnsi::Strncpy(&CStrSourceData[0], static_cast<const char*>(RewriteResultDesc.target.Data()), RewriteResultDesc.target.Size());
 			PreprocessedShaderSource = CStrSourceData.c_str();
 
 			if (bDumpDebugInfo)
@@ -578,12 +575,6 @@ static bool RewriteUsingSC(FString& PreprocessedShaderSource, const FShaderCompi
 				DumpDebugUSF(Input, CStrSourceData.c_str(), 0, GRewrittenBaseFilename);
 			}
 		}
-
-		// Release ShaderConductor resources
-		if (RewriteBlob)
-		{
-			ShaderConductor::DestroyBlob(RewriteBlob);
-		}
 	}
 
 	return bResult;
