Moved FShaderConductorContext into dedicated source files.

#rnx

diff --git a/Engine/Source/Developer/ShaderCompilerCommon/Private/ShaderCompilerCommon.cpp b/Engine/Source/Developer/ShaderCompilerCommon/Private/ShaderCompilerCommon.cpp
--- a/Engine/Source/Developer/ShaderCompilerCommon/Private/ShaderCompilerCommon.cpp
+++ b/Engine/Source/Developer/ShaderCompilerCommon/Private/ShaderCompilerCommon.cpp
@@ -5,19 +5,10 @@
 #include "Modules/ModuleManager.h"
 #include "HlslccDefinitions.h"
 #include "HAL/FileManager.h"
-#include "HAL/ExceptionHandling.h"
-
-#if PLATFORM_MAC || PLATFORM_WINDOWS || PLATFORM_LINUX
-THIRD_PARTY_INCLUDES_START
-	#include "ShaderConductor/ShaderConductor.hpp"
-THIRD_PARTY_INCLUDES_END
-#endif
 
 IMPLEMENT_MODULE(FDefaultModuleImpl, ShaderCompilerCommon);
 
 
-ESCWErrorCode GSCWErrorCode = ESCWErrorCode::NotSet;
-
 int16 GetNumUniformBuffersUsed(const FShaderCompilerResourceTable& InSRT)
 {
 	auto CountLambda = [&](const TArray<uint32>& In)
@@ -2497,757 +2488,4 @@ namespace CrossCompiler
 		return true;
 	}
 
-	/////////// FShaderConductorContext ///////////
-
-#if PLATFORM_MAC || PLATFORM_WINDOWS || PLATFORM_LINUX
-
-	// Inner wrapper function is required here because '__try'-statement cannot be used with function that requires object unwinding
-	static void InnerScRewriteWrapper(
-		const ShaderConductor::Compiler::SourceDesc& InSourceDesc,
-		const ShaderConductor::Compiler::Options& InOptions,
-		ShaderConductor::Compiler::ResultDesc& OutResultDesc)
-	{
-		OutResultDesc = ShaderConductor::Compiler::Rewrite(InSourceDesc, InOptions);
-	}
-
-	static bool ScRewriteWrapper(
-		const ShaderConductor::Compiler::SourceDesc& InSourceDesc,
-		const ShaderConductor::Compiler::Options& InOptions,
-		ShaderConductor::Compiler::ResultDesc& OutResultDesc,
-		bool& bOutException)
-	{
-		bOutException = false;
-#if !PLATFORM_SEH_EXCEPTIONS_DISABLED
-		__try
-#endif
-		{
-			InnerScRewriteWrapper(InSourceDesc, InOptions, OutResultDesc);
-			return true;
-		}
-#if !PLATFORM_SEH_EXCEPTIONS_DISABLED
-		__except (EXCEPTION_EXECUTE_HANDLER)
-		{
-			GSCWErrorCode = ESCWErrorCode::CrashInsidePlatformCompiler;
-			FMemory::Memzero(OutResultDesc);
-			bOutException = true;
-			return false;
-		}
-#endif
-	}
-
-	// Inner wrapper function is required here because '__try'-statement cannot be used with function that requires object unwinding
-	static void InnerScCompileWrapper(
-		const ShaderConductor::Compiler::SourceDesc& InSourceDesc,
-		const ShaderConductor::Compiler::Options& InOptions,
-		const ShaderConductor::Compiler::TargetDesc& InTargetDesc,
-		ShaderConductor::Compiler::ResultDesc& OutResultDesc)
-	{
-		OutResultDesc = ShaderConductor::Compiler::Compile(InSourceDesc, InOptions, InTargetDesc);
-	}
-
-	static bool ScCompileWrapper(
-		const ShaderConductor::Compiler::SourceDesc& InSourceDesc,
-		const ShaderConductor::Compiler::Options& InOptions,
-		const ShaderConductor::Compiler::TargetDesc& InTargetDesc,
-		ShaderConductor::Compiler::ResultDesc& OutResultDesc,
-		bool& bOutException)
-	{
-		bOutException = false;
-#if !PLATFORM_SEH_EXCEPTIONS_DISABLED
-		__try
-#endif
-		{
-			InnerScCompileWrapper(InSourceDesc, InOptions, InTargetDesc, OutResultDesc);
-			return true;
-		}
-#if !PLATFORM_SEH_EXCEPTIONS_DISABLED
-		__except (EXCEPTION_EXECUTE_HANDLER)
-		{
-			GSCWErrorCode = ESCWErrorCode::CrashInsidePlatformCompiler;
-			FMemory::Memzero(OutResultDesc);
-			bOutException = true;
-			return false;
-		}
-#endif
-	}
-
-	// Inner wrapper function is required here because '__try'-statement cannot be used with function that requires object unwinding
-	static void InnerScCompileWrapper(
-		const ShaderConductor::Compiler::ResultDesc& InBinaryDesc,
-		const ShaderConductor::Compiler::SourceDesc& InSourceDesc,
-		const ShaderConductor::Compiler::TargetDesc& InTargetDesc,
-		ShaderConductor::Compiler::ResultDesc& OutResultDesc)
-	{
-		OutResultDesc = ShaderConductor::Compiler::ConvertBinary(InBinaryDesc, InSourceDesc, InTargetDesc);
-	}
-
-	static bool ScConvertBinaryWrapper(
-		const ShaderConductor::Compiler::ResultDesc& InBinaryDesc,
-		const ShaderConductor::Compiler::SourceDesc& InSourceDesc,
-		const ShaderConductor::Compiler::TargetDesc& InTargetDesc,
-		ShaderConductor::Compiler::ResultDesc& OutResultDesc,
-		bool& bOutException)
-	{
-		bOutException = false;
-#if !PLATFORM_SEH_EXCEPTIONS_DISABLED
-		__try
-#endif
-		{
-			InnerScCompileWrapper(InBinaryDesc, InSourceDesc, InTargetDesc, OutResultDesc);
-			return true;
-		}
-#if !PLATFORM_SEH_EXCEPTIONS_DISABLED
-		__except (EXCEPTION_EXECUTE_HANDLER)
-		{
-			GSCWErrorCode = ESCWErrorCode::CrashInsidePlatformCompiler;
-			FMemory::Memzero(OutResultDesc);
-			bOutException = true;
-			return false;
-		}
-#endif
-	}
-
-	// Converts the byte array 'InString' (without null terminator) to the output ANSI string 'OutString' (with appended null terminator).
-	static void ConvertByteArrayToAnsiString(const ANSICHAR* InString, uint32 InStringLength, TArray<ANSICHAR>& OutString)
-	{
-		// 'FCStringAnsi::Strncpy()' will put a '\0' character at the end
-		OutString.SetNum(InStringLength + 1);
-		FCStringAnsi::Strncpy(OutString.GetData(), InString, OutString.Num());
-	}
-
-	// Converts the FString to the output ANSI string 'OutString'.
-	static void ConvertFStringToAnsiString(const FString& InString, TArray<ANSICHAR>& OutString)
-	{
-		ConvertByteArrayToAnsiString(TCHAR_TO_ANSI(*InString), InString.Len(), OutString);
-	}
-
-	// Copies the NULL-terminated string 'InString' to 'OutString'. Also copies the '\0' character at the end.
-	static void CopyAnsiString(const ANSICHAR* InString, TArray<ANSICHAR>& OutString)
-	{
-		// 'InString' is NULL-terminated, so we can use 'FCStringAnsi::Strlen()'
-		if (InString != nullptr)
-		{
-			ConvertByteArrayToAnsiString(InString, FCStringAnsi::Strlen(InString), OutString);
-		}
-	}
-
-	// Converts the specified ShaderConductor blob to FString.
-	static bool ConvertByteArrayToFString(const void* InData, uint32 InSize, FString& OutString)
-	{
-		if (InData != nullptr && InSize > 0)
-		{
-			FUTF8ToTCHAR UTF8Converter(reinterpret_cast<const ANSICHAR*>(InData), InSize);
-			OutString = FString(UTF8Converter.Length(), UTF8Converter.Get());
-			return true;
-		}
-		return false;
-	}
-
-	// Converts the specified ShaderConductor blob to FString.
-	static bool ConvertScBlobToFString(ShaderConductor::Blob* Blob, FString& OutString)
-	{
-		if (Blob && Blob->Size() > 0)
-		{
-			return ConvertByteArrayToFString(Blob->Data(), Blob->Size(), OutString);
-		}
-		return false;
-	}
-
-	static ShaderConductor::ShaderStage ToShaderConductorShaderStage(EHlslShaderFrequency Frequency)
-	{
-		check(Frequency >= HSF_VertexShader && Frequency <= HSF_ComputeShader);
-		switch (Frequency)
-		{
-		case HSF_VertexShader:		return ShaderConductor::ShaderStage::VertexShader;
-		case HSF_PixelShader:		return ShaderConductor::ShaderStage::PixelShader;
-		case HSF_GeometryShader:	return ShaderConductor::ShaderStage::GeometryShader;
-		case HSF_HullShader:		return ShaderConductor::ShaderStage::HullShader;
-		case HSF_DomainShader:		return ShaderConductor::ShaderStage::DomainShader;
-		case HSF_ComputeShader:		return ShaderConductor::ShaderStage::ComputeShader;
-		default:					return ShaderConductor::ShaderStage::NumShaderStages;
-		}
-	}
-
-	static ShaderConductor::Compiler::ShaderModel ToShaderConductorShaderModel(EHlslCompileTarget Target)
-	{
-		switch (Target)
-		{
-		case HCT_FeatureLevelSM4:		return { 4, 0 };
-		case HCT_FeatureLevelES3_1Ext:	return { 4, 0 };
-		case HCT_FeatureLevelSM5:		return { 5, 0 };
-		case HCT_FeatureLevelES3_1:		return { 4, 0 };
-		default: checkf(0, TEXT("Invalid input shader target for enum <EHlslCompileTarget>."));
-		}
-		return { 6,0 };
-	}
-
-	// Wrapper structure to hold all intermediate buffers for ShaderConductor
-	struct FShaderConductorContext::FShaderConductorIntermediates
-	{
-		FShaderConductorIntermediates()
-			: Stage(ShaderConductor::ShaderStage::NumShaderStages)
-		{
-		}
-
-		TArray<ANSICHAR> ShaderSource;
-		TArray<ANSICHAR> Filename;
-		TArray<ANSICHAR> EntryPoint;
-		ShaderConductor::ShaderStage Stage;
-		TArray<TPair<TArray<ANSICHAR>, TArray<ANSICHAR>>> Defines;
-		TArray<ShaderConductor::MacroDefine> DefineRefs;
-		TArray<TPair<TArray<ANSICHAR>, TArray<ANSICHAR>>> Flags;
-		TArray<ShaderConductor::MacroDefine> FlagRefs;
-	};
-
-	static void ConvertScSourceDesc(const FShaderConductorContext::FShaderConductorIntermediates& Intermediates, ShaderConductor::Compiler::SourceDesc& OutSourceDesc)
-	{
-		// Convert descriptor with pointers to the ANSI strings
-		OutSourceDesc.source = Intermediates.ShaderSource.GetData();
-		OutSourceDesc.fileName = Intermediates.Filename.GetData();
-		OutSourceDesc.entryPoint = Intermediates.EntryPoint.GetData();
-		OutSourceDesc.stage = Intermediates.Stage;
-		if (Intermediates.DefineRefs.Num() > 0)
-		{
-			OutSourceDesc.defines = Intermediates.DefineRefs.GetData();
-			OutSourceDesc.numDefines = static_cast<uint32>(Intermediates.DefineRefs.Num());
-		}
-		else
-		{
-			OutSourceDesc.defines = nullptr;
-			OutSourceDesc.numDefines = 0;
-		}
-	}
-
-	static const ANSICHAR* GetGlslFamilyVersionString(int32 Version)
-	{
-		switch (Version)
-		{
-		case 310: return "310";
-		case 320: return "320";
-		case 330: return "330";
-		case 430: return "430";
-		default: return nullptr;
-		}
-	}
-
-	static void ConvertScTargetDescLanguageGlslFamily(const FShaderConductorTarget& InTarget, ShaderConductor::Compiler::TargetDesc& OutTargetDesc)
-	{
-		OutTargetDesc.language = (InTarget.Language == EShaderConductorLanguage::Glsl ? ShaderConductor::ShadingLanguage::Glsl : ShaderConductor::ShadingLanguage::Essl);
-		OutTargetDesc.version = GetGlslFamilyVersionString(InTarget.Version);
-		checkf(OutTargetDesc.version, TEXT("Unsupported target shader version for GLSL family: %d"), InTarget.Version);
-	}
-
-	static const ANSICHAR* GetMetalFamilyVersionString(int32 Version)
-	{
-		switch (Version)
-		{
-		case 20100: return "20100";
-		case 20000: return "20000";
-		case 10200: return "10200";
-		case 10100: return "10100";
-		case 10000: return "10000";
-		default: return nullptr;
-		}
-	}
-
-	static void ConvertScTargetDescLanguageMetalFamily(const FShaderConductorTarget& InTarget, ShaderConductor::Compiler::TargetDesc& OutTargetDesc)
-	{
-		OutTargetDesc.language = (InTarget.Language == EShaderConductorLanguage::Metal_macOS ? ShaderConductor::ShadingLanguage::Msl_macOS : ShaderConductor::ShadingLanguage::Msl_iOS);
-		OutTargetDesc.version = GetMetalFamilyVersionString(InTarget.Version);
-		checkf(OutTargetDesc.version, TEXT("Unsupported target shader version for Metal family: %d"), InTarget.Version);
-	}
-
-	static void ConvertScTargetDesc(FShaderConductorContext::FShaderConductorIntermediates& Intermediates, const FShaderConductorTarget& InTarget, ShaderConductor::Compiler::TargetDesc& OutTargetDesc)
-	{
-		// Convert FString to ANSI string and store them as intermediates
-		FMemory::Memzero(OutTargetDesc);
-
-		switch (InTarget.Language)
-		{
-		case EShaderConductorLanguage::Glsl:
-		case EShaderConductorLanguage::Essl:
-			ConvertScTargetDescLanguageGlslFamily(InTarget, OutTargetDesc);
-			break;
-		case EShaderConductorLanguage::Metal_macOS:
-		case EShaderConductorLanguage::Metal_iOS:
-			ConvertScTargetDescLanguageMetalFamily(InTarget, OutTargetDesc);
-			break;
-		}
-
-		// Convert flags map into an array container
-		TArray<ANSICHAR> FlagName, FlagValue;
-		for (const TPair<FString, FString>& Iter : InTarget.CompileFlags.GetDefinitionMap())
-		{
-			ConvertFStringToAnsiString(Iter.Key, FlagName);
-			ConvertFStringToAnsiString(Iter.Value, FlagValue);
-			Intermediates.Flags.Emplace(MoveTemp(FlagName), MoveTemp(FlagValue));
-		}
-
-		// Store references after all elements have been added to the container so the pointers remain valid
-		Intermediates.FlagRefs.SetNum(Intermediates.Flags.Num());
-		for (int32 Index = 0; Index < Intermediates.Flags.Num(); ++Index)
-		{
-			Intermediates.FlagRefs[Index].name = Intermediates.Flags[Index].Key.GetData();
-			Intermediates.FlagRefs[Index].value = Intermediates.Flags[Index].Value.GetData();
-		}
-
-		OutTargetDesc.options = Intermediates.FlagRefs.GetData();
-		OutTargetDesc.numOptions = static_cast<uint32>(Intermediates.FlagRefs.Num());
-
-		// Wrap input function into lambda to convert to ShaderConductor interface
-		if (InTarget.VariableTypeRenameCallback)
-		{
-			OutTargetDesc.variableTypeRenameCallback = [InnerCallback = InTarget.VariableTypeRenameCallback](const char* VariableName, const char* TypeName) -> ShaderConductor::Blob
-			{
-				// Forward callback to public interface callback
-				FString RenamedTypeName;
-				if (InnerCallback(FAnsiStringView(VariableName), FAnsiStringView(TypeName), RenamedTypeName))
-				{
-					if (!RenamedTypeName.IsEmpty())
-					{
-						// Convert renamed type name from FString to ShaderConductor::Blob
-						return ShaderConductor::Blob(TCHAR_TO_ANSI(*RenamedTypeName), RenamedTypeName.Len() + 1);
-					}
-				}
-				return ShaderConductor::Blob{};
-			};
-		}
-	}
-
-	static void ConvertScOptions(const FShaderConductorOptions& InOptions, ShaderConductor::Compiler::Options& OutOptions)
-	{
-		OutOptions.removeUnusedGlobals = InOptions.bRemoveUnusedGlobals;
-		OutOptions.packMatricesInRowMajor = InOptions.bPackMatricesInRowMajor;
-		OutOptions.enable16bitTypes = InOptions.bEnable16bitTypes;
-		OutOptions.enableDebugInfo = InOptions.bEnableDebugInfo;
-		OutOptions.disableOptimizations = InOptions.bDisableOptimizations;
-		OutOptions.enableFMAPass = InOptions.bEnableFMAPass;
-		OutOptions.shaderModel = ToShaderConductorShaderModel(InOptions.TargetProfile);
-	}
-
-	static void ConvertScDefines(FShaderConductorContext::FShaderConductorIntermediates& Intermediates, const FShaderCompilerDefinitions& InDefinitions)
-	{
-		// Convert FString to ANSI string for each macro definition and its value
-		TArray<ANSICHAR> DefineName, DefineValue;
-
-		for (const TPair<FString, FString>& Iter : InDefinitions.GetDefinitionMap())
-		{
-			ConvertFStringToAnsiString(Iter.Key, DefineName);
-			ConvertFStringToAnsiString(Iter.Value, DefineValue);
-			Intermediates.Defines.Emplace(MoveTemp(DefineName), MoveTemp(DefineValue));
-		}
-
-		// Store references after all elements have been added to the container so the pointers remain valid
-		Intermediates.DefineRefs.SetNum(Intermediates.Defines.Num());
-		for (int32 Index = 0; Index < Intermediates.Defines.Num(); ++Index)
-		{
-			Intermediates.DefineRefs[Index].name = Intermediates.Defines[Index].Key.GetData();
-			Intermediates.DefineRefs[Index].value = Intermediates.Defines[Index].Value.GetData();
-		}
-	}
-
-	// Returns whether the specified line of text contains only these characters, making it a valid line marker from DXC: ' ', '\t', '~', '^'
-	static bool IsTextLineDxcLineMarker(const FString& Line)
-	{
-		bool bContainsLineMarkerChars = false;
-		for (TCHAR Char : Line)
-		{
-			if (Char == TCHAR('~') || Char == TCHAR('^'))
-			{
-				// Line contains at least one of the necessary characters to be a potential DXC line marker.
-				bContainsLineMarkerChars = true;
-			}
-			else if (!(Char == TCHAR(' ') || Char == TCHAR('\t')))
-			{
-				// Illegal character for a potential DXC line marker.
-				return false;
-			}
-		}
-		return bContainsLineMarkerChars;
-	}
-
-	// Converts the error blob from ShaderConductor into an array of error reports (of type FShaderCompilerError).
-	static void ConvertScCompileErrors(ShaderConductor::Blob& ErrorBlob, TArray<FShaderCompilerError>& OutErrors)
-	{
-		// Convert blob into FString
-		FString ErrorString;
-		if (ConvertScBlobToFString(&ErrorBlob, ErrorString))
-		{
-			// Convert FString into array of FString (one for each line)
-			TArray<FString> ErrorStringLines;
-			ErrorString.ParseIntoArray(ErrorStringLines, TEXT("\n"));
-
-			// Forward parsed array of lines to primary conversion function
-			FShaderConductorContext::ConvertCompileErrors(MoveTemp(ErrorStringLines), OutErrors);
-		}
-	}
-
-	FShaderConductorContext::FShaderConductorContext()
-		: Intermediates(new FShaderConductorIntermediates())
-	{
-	}
-
-	FShaderConductorContext::~FShaderConductorContext()
-	{
-		delete Intermediates;
-	}
-
-	FShaderConductorContext::FShaderConductorContext(FShaderConductorContext&& Rhs)
-		: Errors(MoveTemp(Rhs.Errors))
-		, Intermediates(Rhs.Intermediates)
-	{
-		Rhs.Intermediates = nullptr;
-	}
-
-	FShaderConductorContext& FShaderConductorContext::operator = (FShaderConductorContext&& Rhs)
-	{
-		Errors = MoveTemp(Rhs.Errors);
-		delete Intermediates;
-		Intermediates = Rhs.Intermediates;
-		Rhs.Intermediates = nullptr;
-		return *this;
-	}
-
-	bool FShaderConductorContext::LoadSource(const FString& ShaderSource, const FString& Filename, const FString& EntryPoint, EHlslShaderFrequency ShaderStage, const FShaderCompilerDefinitions* Definitions)
-	{
-		// Convert FString to ANSI string and store them as intermediates
-		ConvertFStringToAnsiString(ShaderSource, Intermediates->ShaderSource);
-		ConvertFStringToAnsiString(Filename, Intermediates->Filename);
-		ConvertFStringToAnsiString(EntryPoint, Intermediates->EntryPoint);
-
-		// Convert macro definitions map into an array container
-		if (Definitions != nullptr)
-		{
-			ConvertScDefines(*Intermediates, *Definitions);
-		}
-
-		// Convert shader stage
-		Intermediates->Stage = ToShaderConductorShaderStage(ShaderStage);
-
-		return true;
-	}
-
-	bool FShaderConductorContext::LoadSource(const ANSICHAR* ShaderSource, const ANSICHAR* Filename, const ANSICHAR* EntryPoint, EHlslShaderFrequency ShaderStage, const FShaderCompilerDefinitions* Definitions)
-	{
-		// Store ANSI strings as intermediates
-		CopyAnsiString(ShaderSource, Intermediates->ShaderSource);
-		CopyAnsiString(Filename, Intermediates->Filename);
-		CopyAnsiString(EntryPoint, Intermediates->EntryPoint);
-
-		// Convert macro definitions map into an array container
-		if (Definitions != nullptr)
-		{
-			ConvertScDefines(*Intermediates, *Definitions);
-		}
-
-		// Convert shader stage
-		Intermediates->Stage = ToShaderConductorShaderStage(ShaderStage);
-
-		return true;
-	}
-
-	bool FShaderConductorContext::RewriteHlsl(const FShaderConductorOptions& Options, FString* OutSource)
-	{
-		// Convert descriptors for ShaderConductor interface
-		ShaderConductor::Compiler::SourceDesc ScSourceDesc;
-		ConvertScSourceDesc(*Intermediates, ScSourceDesc);
-
-		ShaderConductor::Compiler::Options ScOptions;
-		ConvertScOptions(Options, ScOptions);
-
-		// Rewrite HLSL with wrapper function to catch exceptions from ShaderConductor
-		bool bSucceeded = false;
-		bool bException = false;
-		ShaderConductor::Compiler::ResultDesc ResultDesc;
-		ScRewriteWrapper(ScSourceDesc, ScOptions, ResultDesc, bException);
-
-		if (!ResultDesc.hasError && !bException && ResultDesc.target.Size() > 0)
-		{
-			// Copy rewritten HLSL code into intermediate source code.
-			ConvertByteArrayToAnsiString(reinterpret_cast<const ANSICHAR*>(ResultDesc.target.Data()), ResultDesc.target.Size(), Intermediates->ShaderSource);
-
-			// If output source is specified, also convert to TCHAR string
-			if (OutSource != nullptr)
-			{
-				*OutSource = ANSI_TO_TCHAR(Intermediates->ShaderSource.GetData());
-			}
-			bSucceeded = true;
-		}
-		else
-		{
-			if (bException)
-			{
-				Errors.Add(TEXT("ShaderConductor exception during rewrite"));
-			}
-			bSucceeded = false;
-		}
-
-		// Append compile error and warning to output reports
-		ConvertScCompileErrors(ResultDesc.errorWarningMsg, Errors);
-
-		return bSucceeded;
-	}
-
-	bool FShaderConductorContext::CompileHlslToSpirv(const FShaderConductorOptions& Options, TArray<uint32>& OutSpirv)
-	{
-		// Convert descriptors for ShaderConductor interface
-		ShaderConductor::Compiler::SourceDesc ScSourceDesc;
-		ConvertScSourceDesc(*Intermediates, ScSourceDesc);
-
-		ShaderConductor::Compiler::TargetDesc ScTargetDesc;
-		FMemory::Memzero(ScTargetDesc);
-		ScTargetDesc.language = ShaderConductor::ShadingLanguage::SpirV;
-
-		ShaderConductor::Compiler::Options ScOptions;
-		ConvertScOptions(Options, ScOptions);
-
-		// Compile HLSL source code to SPIR-V
-		bool bSucceeded = false;
-		bool bException = false;
-		ShaderConductor::Compiler::ResultDesc ResultDesc;
-		ScCompileWrapper(ScSourceDesc, ScOptions, ScTargetDesc, ResultDesc, bException);
-
-		if (!ResultDesc.hasError && !bException && ResultDesc.target.Size() > 0)
-		{
-			// Copy result blob into output SPIR-V module
-			OutSpirv = TArray<uint32>(reinterpret_cast<const uint32*>(ResultDesc.target.Data()), ResultDesc.target.Size() / 4);
-			bSucceeded = true;
-		}
-		else
-		{
-			if (bException)
-			{
-				Errors.Add(TEXT("ShaderConductor exception during compilation"));
-			}
-			bSucceeded = false;
-		}
-
-		// Append compile error and warning to output reports
-		ConvertScCompileErrors(ResultDesc.errorWarningMsg, Errors);
-
-		return bSucceeded;
-	}
-
-	bool FShaderConductorContext::CompileSpirvToSource(const FShaderConductorOptions& Options, const FShaderConductorTarget& Target, const void* InSpirv, uint32 InSpirvByteSize, FString& OutSource)
-	{
-		return CompileSpirvToSourceBuffer(
-			Options, Target, InSpirv, InSpirvByteSize,
-			[&OutSource](const void* Data, uint32 Size)
-			{
-				// Convert source buffer to FString
-				ConvertByteArrayToFString(Data, Size, OutSource);
-			}
-		);
-	}
-
-	bool FShaderConductorContext::CompileSpirvToSourceAnsi(const FShaderConductorOptions& Options, const FShaderConductorTarget& Target, const void* InSpirv, uint32 InSpirvByteSize, TArray<ANSICHAR>& OutSource)
-	{
-		return CompileSpirvToSourceBuffer(
-			Options, Target, InSpirv, InSpirvByteSize,
-			[&OutSource](const void* Data, uint32 Size)
-			{
-				// Convert source buffer to ANSI string
-				ConvertByteArrayToAnsiString(reinterpret_cast<const ANSICHAR*>(Data), Size, OutSource);
-			}
-		);
-	}
-
-	bool FShaderConductorContext::CompileSpirvToSourceBuffer(const FShaderConductorOptions& Options, const FShaderConductorTarget& Target, const void* InSpirv, uint32 InSpirvByteSize, const TFunction<void(const void* Data, uint32 Size)>& OutputCallback)
-	{
-		check(OutputCallback != nullptr);
-		check(InSpirv != nullptr);
-		check(InSpirvByteSize > 0);
-		checkf(InSpirvByteSize % 4 == 0, TEXT("SPIR-V code unaligned. Size must be a multiple of 4, but %u was specified."), InSpirvByteSize);
-
-		// Convert descriptors for ShaderConductor interface
-		ShaderConductor::Compiler::SourceDesc ScSourceDesc;
-		ConvertScSourceDesc(*Intermediates, ScSourceDesc);
-
-		ShaderConductor::Compiler::TargetDesc ScTargetDesc;
-		ConvertScTargetDesc(*Intermediates, Target, ScTargetDesc);
-
-		ShaderConductor::Compiler::Options ScOptions;
-		ConvertScOptions(Options, ScOptions);
-
-		ShaderConductor::Compiler::ResultDesc ScBinaryDesc;
-		ScBinaryDesc.target.Reset(InSpirv, InSpirvByteSize);
-		ScBinaryDesc.isText = false;
-		ScBinaryDesc.hasError = false;
-
-		// Convert the input SPIR-V into Metal high level source
-		bool bSucceeded = false;
-		bool bException = false;
-		ShaderConductor::Compiler::ResultDesc ResultDesc;
-		ScConvertBinaryWrapper(ScBinaryDesc, ScSourceDesc, ScTargetDesc, ResultDesc, bException);
-
-		if (!ResultDesc.hasError && !bException && ResultDesc.target.Size() > 0)
-		{
-			// Copy result blob into output SPIR-V module
-			OutputCallback(ResultDesc.target.Data(), ResultDesc.target.Size());
-			bSucceeded = true;
-		}
-		else
-		{
-			if (bException)
-			{
-				Errors.Add(TEXT("ShaderConductor exception during SPIR-V binary conversion"));
-			}
-			bSucceeded = false;
-		}
-
-		// Append compile error and warning to output reports
-		if (ResultDesc.errorWarningMsg.Size() > 0)
-		{
-			FString ErrorString;
-			if (ConvertScBlobToFString(&ResultDesc.errorWarningMsg, ErrorString))
-			{
-				Errors.Add(*ErrorString);
-			}
-		}
-
-		return bSucceeded;
-	}
-
-	void FShaderConductorContext::FlushErrors(TArray<FShaderCompilerError>& OutErrors)
-	{
-		if (OutErrors.Num() > 0)
-		{
-			// Append internal list of errors to output list, then clear internal list
-			for (const FShaderCompilerError& ErrorEntry : Errors)
-			{
-				OutErrors.Add(ErrorEntry);
-			}
-			Errors.Empty();
-		}
-		else
-		{
-			// Move internal list of errors into output list
-			OutErrors = MoveTemp(Errors);
-		}
-	}
-
-	const ANSICHAR* FShaderConductorContext::GetSourceString() const
-	{
-		return (Intermediates->ShaderSource.Num() > 0 ? Intermediates->ShaderSource.GetData() : nullptr);
-	}
-
-	int32 FShaderConductorContext::GetSourceLength() const
-	{
-		return (Intermediates->ShaderSource.Num() > 0 ? (Intermediates->ShaderSource.Num() - 1) : 0);
-	}
-
-	void FShaderConductorContext::ConvertCompileErrors(TArray<FString>&& ErrorStringLines, TArray<FShaderCompilerError>& OutErrors)
-	{
-		// Returns whether the specified line in the 'ErrorStringLines' array has a line marker.
-		auto HasErrorLineMarker = [&ErrorStringLines](int32 LineIndex)
-		{
-			if (LineIndex + 2 < ErrorStringLines.Num())
-			{
-				return IsTextLineDxcLineMarker(ErrorStringLines[LineIndex + 2]);
-			}
-			return false;
-		};
-
-		// Iterate over all errors. Most (but not all) contain a highlighted line and line marker.
-		for (int32 LineIndex = 0; LineIndex < ErrorStringLines.Num();)
-		{
-			if (HasErrorLineMarker(LineIndex))
-			{
-				// Add current line as error with highlighted source line (LineIndex+1) and line marker (LineIndex+2)
-				OutErrors.Emplace(MoveTemp(ErrorStringLines[LineIndex]), MoveTemp(ErrorStringLines[LineIndex + 1]), MoveTemp(ErrorStringLines[LineIndex + 2]));
-				LineIndex += 3;
-			}
-			else
-			{
-				// Add current line as single error
-				OutErrors.Emplace(MoveTemp(ErrorStringLines[LineIndex]));
-				LineIndex += 1;
-			}
-		}
-	}
-
-#else // PLATFORM_MAC || PLATFORM_WINDOWS || PLATFORM_LINUX
-
-	FShaderConductorContext::FShaderConductorContext()
-	{
-		checkf(0, TEXT("Cannot instantiate FShaderConductorContext for unsupported platform"));
-	}
-
-	FShaderConductorContext::~FShaderConductorContext()
-	{
-		// Dummy
-	}
-
-	FShaderConductorContext::FShaderConductorContext(FShaderConductorContext&& Rhs)
-	{
-		// Dummy
-	}
-
-	FShaderConductorContext& FShaderConductorContext::operator = (FShaderConductorContext&& Rhs)
-	{
-		return *this; // Dummy
-	}
-
-	bool FShaderConductorContext::LoadSource(const FString& ShaderSource, const FString& Filename, const FString& EntryPoint, EHlslShaderFrequency ShaderStage, const FShaderCompilerDefinitions* Definitions)
-	{
-		return false; // Dummy
-	}
-
-	bool FShaderConductorContext::LoadSource(const ANSICHAR* ShaderSource, const ANSICHAR* Filename, const ANSICHAR* EntryPoint, EHlslShaderFrequency ShaderStage, const FShaderCompilerDefinitions* Definitions)
-	{
-		return false; // Dummy
-	}
-
-	bool FShaderConductorContext::RewriteHlsl(const FShaderConductorOptions& Options, FString* OutSource)
-	{
-		return false; // Dummy
-	}
-
-	bool FShaderConductorContext::CompileHlslToSpirv(const FShaderConductorOptions& Options, TArray<uint32>& OutSpirv)
-	{
-		return false; // Dummy
-	}
-
-	bool FShaderConductorContext::CompileSpirvToSource(const FShaderConductorOptions& Options, const FShaderConductorTarget& Target, const void* InSpirv, uint32 InSpirvByteSize, FString& OutSource)
-	{
-		return false; // Dummy
-	}
-
-	bool FShaderConductorContext::CompileSpirvToSourceAnsi(const FShaderConductorOptions& Options, const FShaderConductorTarget& Target, const void* InSpirv, uint32 InSpirvByteSize, TArray<ANSICHAR>& OutSource)
-	{
-		return false; // Dummy
-	}
-
-	bool FShaderConductorContext::CompileSpirvToSourceBuffer(const FShaderConductorOptions& Options, const FShaderConductorTarget& Target, const void* InSpirv, uint32 InSpirvByteSize, const TFunction<void(const void* Data, uint32 Size)>& OutputCallback)
-	{
-		return false; // Dummy
-	}
-
-	void FShaderConductorContext::FlushErrors(TArray<FShaderCompilerError>& OutErrors)
-	{
-		// Dummy
-	}
-
-	const ANSICHAR* FShaderConductorContext::GetSourceString() const
-	{
-		return nullptr; // Dummy
-	}
-
-	int32 FShaderConductorContext::GetSourceLength() const
-	{
-		return 0; // Dummy
-	}
-
-	void FShaderConductorContext::ConvertCompileErrors(const TArray<FString>& ErrorStringLines, TArray<FShaderCompilerError>& OutErrors)
-	{
-		// Dummy
-	}
-
-#endif // PLATFORM_MAC || PLATFORM_WINDOWS || PLATFORM_LINUX
-
-	bool FShaderConductorContext::IsIntermediateSpirvOutputVariable(const ANSICHAR* SpirvVariableName)
-	{
-		// This is only true for "temp.var.hullMainRetVal" which is generated by DXC as intermediate output variable to communicate patch constant data in a Hull Shader.
-		return (SpirvVariableName != nullptr && FCStringAnsi::Strcmp(SpirvVariableName, "temp.var.hullMainRetVal") == 0);
-	}
-}
+} // namespace CrossCompiler
diff --git a/Engine/Source/Developer/ShaderCompilerCommon/Private/ShaderConductorContext.cpp b/Engine/Source/Developer/ShaderCompilerCommon/Private/ShaderConductorContext.cpp
new file mode 100644
--- /dev/null
+++ b/Engine/Source/Developer/ShaderCompilerCommon/Private/ShaderConductorContext.cpp
@@ -0,0 +1,778 @@
+// Copyright Epic Games, Inc. All Rights Reserved.
+
+#include "ShaderConductorContext.h"
+#include "HAL/ExceptionHandling.h"
+
+#if PLATFORM_MAC || PLATFORM_WINDOWS || PLATFORM_LINUX
+THIRD_PARTY_INCLUDES_START
+	#include "ShaderConductor/ShaderConductor.hpp"
+THIRD_PARTY_INCLUDES_END
+#endif
+
+
+ESCWErrorCode GSCWErrorCode = ESCWErrorCode::NotSet;
+
+namespace CrossCompiler
+{
+
+#if PLATFORM_MAC || PLATFORM_WINDOWS || PLATFORM_LINUX
+
+	// Inner wrapper function is required here because '__try'-statement cannot be used with function that requires object unwinding
+	static void InnerScRewriteWrapper(
+		const ShaderConductor::Compiler::SourceDesc& InSourceDesc,
+		const ShaderConductor::Compiler::Options& InOptions,
+		ShaderConductor::Compiler::ResultDesc& OutResultDesc)
+	{
+		OutResultDesc = ShaderConductor::Compiler::Rewrite(InSourceDesc, InOptions);
+	}
+
+	static bool ScRewriteWrapper(
+		const ShaderConductor::Compiler::SourceDesc& InSourceDesc,
+		const ShaderConductor::Compiler::Options& InOptions,
+		ShaderConductor::Compiler::ResultDesc& OutResultDesc,
+		bool& bOutException)
+	{
+		bOutException = false;
+#if !PLATFORM_SEH_EXCEPTIONS_DISABLED
+		__try
+#endif
+		{
+			InnerScRewriteWrapper(InSourceDesc, InOptions, OutResultDesc);
+			return true;
+		}
+#if !PLATFORM_SEH_EXCEPTIONS_DISABLED
+		__except (EXCEPTION_EXECUTE_HANDLER)
+		{
+			GSCWErrorCode = ESCWErrorCode::CrashInsidePlatformCompiler;
+			FMemory::Memzero(OutResultDesc);
+			bOutException = true;
+			return false;
+		}
+#endif
+	}
+
+	// Inner wrapper function is required here because '__try'-statement cannot be used with function that requires object unwinding
+	static void InnerScCompileWrapper(
+		const ShaderConductor::Compiler::SourceDesc& InSourceDesc,
+		const ShaderConductor::Compiler::Options& InOptions,
+		const ShaderConductor::Compiler::TargetDesc& InTargetDesc,
+		ShaderConductor::Compiler::ResultDesc& OutResultDesc)
+	{
+		OutResultDesc = ShaderConductor::Compiler::Compile(InSourceDesc, InOptions, InTargetDesc);
+	}
+
+	static bool ScCompileWrapper(
+		const ShaderConductor::Compiler::SourceDesc& InSourceDesc,
+		const ShaderConductor::Compiler::Options& InOptions,
+		const ShaderConductor::Compiler::TargetDesc& InTargetDesc,
+		ShaderConductor::Compiler::ResultDesc& OutResultDesc,
+		bool& bOutException)
+	{
+		bOutException = false;
+#if !PLATFORM_SEH_EXCEPTIONS_DISABLED
+		__try
+#endif
+		{
+			InnerScCompileWrapper(InSourceDesc, InOptions, InTargetDesc, OutResultDesc);
+			return true;
+		}
+#if !PLATFORM_SEH_EXCEPTIONS_DISABLED
+		__except (EXCEPTION_EXECUTE_HANDLER)
+		{
+			GSCWErrorCode = ESCWErrorCode::CrashInsidePlatformCompiler;
+			FMemory::Memzero(OutResultDesc);
+			bOutException = true;
+			return false;
+		}
+#endif
+	}
+
+	// Inner wrapper function is required here because '__try'-statement cannot be used with function that requires object unwinding
+	static void InnerScCompileWrapper(
+		const ShaderConductor::Compiler::ResultDesc& InBinaryDesc,
+		const ShaderConductor::Compiler::SourceDesc& InSourceDesc,
+		const ShaderConductor::Compiler::TargetDesc& InTargetDesc,
+		ShaderConductor::Compiler::ResultDesc& OutResultDesc)
+	{
+		OutResultDesc = ShaderConductor::Compiler::ConvertBinary(InBinaryDesc, InSourceDesc, InTargetDesc);
+	}
+
+	static bool ScConvertBinaryWrapper(
+		const ShaderConductor::Compiler::ResultDesc& InBinaryDesc,
+		const ShaderConductor::Compiler::SourceDesc& InSourceDesc,
+		const ShaderConductor::Compiler::TargetDesc& InTargetDesc,
+		ShaderConductor::Compiler::ResultDesc& OutResultDesc,
+		bool& bOutException)
+	{
+		bOutException = false;
+#if !PLATFORM_SEH_EXCEPTIONS_DISABLED
+		__try
+#endif
+		{
+			InnerScCompileWrapper(InBinaryDesc, InSourceDesc, InTargetDesc, OutResultDesc);
+			return true;
+		}
+#if !PLATFORM_SEH_EXCEPTIONS_DISABLED
+		__except (EXCEPTION_EXECUTE_HANDLER)
+		{
+			GSCWErrorCode = ESCWErrorCode::CrashInsidePlatformCompiler;
+			FMemory::Memzero(OutResultDesc);
+			bOutException = true;
+			return false;
+		}
+#endif
+	}
+
+	// Converts the byte array 'InString' (without null terminator) to the output ANSI string 'OutString' (with appended null terminator).
+	static void ConvertByteArrayToAnsiString(const ANSICHAR* InString, uint32 InStringLength, TArray<ANSICHAR>& OutString)
+	{
+		// 'FCStringAnsi::Strncpy()' will put a '\0' character at the end
+		OutString.SetNum(InStringLength + 1);
+		FCStringAnsi::Strncpy(OutString.GetData(), InString, OutString.Num());
+	}
+
+	// Converts the FString to the output ANSI string 'OutString'.
+	static void ConvertFStringToAnsiString(const FString& InString, TArray<ANSICHAR>& OutString)
+	{
+		ConvertByteArrayToAnsiString(TCHAR_TO_ANSI(*InString), InString.Len(), OutString);
+	}
+
+	// Copies the NULL-terminated string 'InString' to 'OutString'. Also copies the '\0' character at the end.
+	static void CopyAnsiString(const ANSICHAR* InString, TArray<ANSICHAR>& OutString)
+	{
+		// 'InString' is NULL-terminated, so we can use 'FCStringAnsi::Strlen()'
+		if (InString != nullptr)
+		{
+			ConvertByteArrayToAnsiString(InString, FCStringAnsi::Strlen(InString), OutString);
+		}
+	}
+
+	// Converts the specified ShaderConductor blob to FString.
+	static bool ConvertByteArrayToFString(const void* InData, uint32 InSize, FString& OutString)
+	{
+		if (InData != nullptr && InSize > 0)
+		{
+			FUTF8ToTCHAR UTF8Converter(reinterpret_cast<const ANSICHAR*>(InData), InSize);
+			OutString = FString(UTF8Converter.Length(), UTF8Converter.Get());
+			return true;
+		}
+		return false;
+	}
+
+	// Converts the specified ShaderConductor blob to FString.
+	static bool ConvertScBlobToFString(ShaderConductor::Blob* Blob, FString& OutString)
+	{
+		if (Blob && Blob->Size() > 0)
+		{
+			return ConvertByteArrayToFString(Blob->Data(), Blob->Size(), OutString);
+		}
+		return false;
+	}
+
+	static ShaderConductor::ShaderStage ToShaderConductorShaderStage(EHlslShaderFrequency Frequency)
+	{
+		check(Frequency >= HSF_VertexShader && Frequency <= HSF_ComputeShader);
+		switch (Frequency)
+		{
+		case HSF_VertexShader:		return ShaderConductor::ShaderStage::VertexShader;
+		case HSF_PixelShader:		return ShaderConductor::ShaderStage::PixelShader;
+		case HSF_GeometryShader:	return ShaderConductor::ShaderStage::GeometryShader;
+		case HSF_HullShader:		return ShaderConductor::ShaderStage::HullShader;
+		case HSF_DomainShader:		return ShaderConductor::ShaderStage::DomainShader;
+		case HSF_ComputeShader:		return ShaderConductor::ShaderStage::ComputeShader;
+		default:					return ShaderConductor::ShaderStage::NumShaderStages;
+		}
+	}
+
+	static ShaderConductor::Compiler::ShaderModel ToShaderConductorShaderModel(EHlslCompileTarget Target)
+	{
+		switch (Target)
+		{
+		case HCT_FeatureLevelSM4:		return { 4, 0 };
+		case HCT_FeatureLevelES3_1Ext:	return { 4, 0 };
+		case HCT_FeatureLevelSM5:		return { 5, 0 };
+		case HCT_FeatureLevelES3_1:		return { 4, 0 };
+		default: checkf(0, TEXT("Invalid input shader target for enum <EHlslCompileTarget>."));
+		}
+		return { 6,0 };
+	}
+
+	// Wrapper structure to hold all intermediate buffers for ShaderConductor
+	struct FShaderConductorContext::FShaderConductorIntermediates
+	{
+		FShaderConductorIntermediates()
+			: Stage(ShaderConductor::ShaderStage::NumShaderStages)
+		{
+		}
+
+		TArray<ANSICHAR> ShaderSource;
+		TArray<ANSICHAR> Filename;
+		TArray<ANSICHAR> EntryPoint;
+		ShaderConductor::ShaderStage Stage;
+		TArray<TPair<TArray<ANSICHAR>, TArray<ANSICHAR>>> Defines;
+		TArray<ShaderConductor::MacroDefine> DefineRefs;
+		TArray<TPair<TArray<ANSICHAR>, TArray<ANSICHAR>>> Flags;
+		TArray<ShaderConductor::MacroDefine> FlagRefs;
+	};
+
+	static void ConvertScSourceDesc(const FShaderConductorContext::FShaderConductorIntermediates& Intermediates, ShaderConductor::Compiler::SourceDesc& OutSourceDesc)
+	{
+		// Convert descriptor with pointers to the ANSI strings
+		OutSourceDesc.source = Intermediates.ShaderSource.GetData();
+		OutSourceDesc.fileName = Intermediates.Filename.GetData();
+		OutSourceDesc.entryPoint = Intermediates.EntryPoint.GetData();
+		OutSourceDesc.stage = Intermediates.Stage;
+		if (Intermediates.DefineRefs.Num() > 0)
+		{
+			OutSourceDesc.defines = Intermediates.DefineRefs.GetData();
+			OutSourceDesc.numDefines = static_cast<uint32>(Intermediates.DefineRefs.Num());
+		}
+		else
+		{
+			OutSourceDesc.defines = nullptr;
+			OutSourceDesc.numDefines = 0;
+		}
+	}
+
+	static const ANSICHAR* GetGlslFamilyVersionString(int32 Version)
+	{
+		switch (Version)
+		{
+		case 310: return "310";
+		case 320: return "320";
+		case 330: return "330";
+		case 430: return "430";
+		default: return nullptr;
+		}
+	}
+
+	static void ConvertScTargetDescLanguageGlslFamily(const FShaderConductorTarget& InTarget, ShaderConductor::Compiler::TargetDesc& OutTargetDesc)
+	{
+		OutTargetDesc.language = (InTarget.Language == EShaderConductorLanguage::Glsl ? ShaderConductor::ShadingLanguage::Glsl : ShaderConductor::ShadingLanguage::Essl);
+		OutTargetDesc.version = GetGlslFamilyVersionString(InTarget.Version);
+		checkf(OutTargetDesc.version, TEXT("Unsupported target shader version for GLSL family: %d"), InTarget.Version);
+	}
+
+	static const ANSICHAR* GetMetalFamilyVersionString(int32 Version)
+	{
+		switch (Version)
+		{
+		case 20100: return "20100";
+		case 20000: return "20000";
+		case 10200: return "10200";
+		case 10100: return "10100";
+		case 10000: return "10000";
+		default: return nullptr;
+		}
+	}
+
+	static void ConvertScTargetDescLanguageMetalFamily(const FShaderConductorTarget& InTarget, ShaderConductor::Compiler::TargetDesc& OutTargetDesc)
+	{
+		OutTargetDesc.language = (InTarget.Language == EShaderConductorLanguage::Metal_macOS ? ShaderConductor::ShadingLanguage::Msl_macOS : ShaderConductor::ShadingLanguage::Msl_iOS);
+		OutTargetDesc.version = GetMetalFamilyVersionString(InTarget.Version);
+		checkf(OutTargetDesc.version, TEXT("Unsupported target shader version for Metal family: %d"), InTarget.Version);
+	}
+
+	static void ConvertScTargetDesc(FShaderConductorContext::FShaderConductorIntermediates& Intermediates, const FShaderConductorTarget& InTarget, ShaderConductor::Compiler::TargetDesc& OutTargetDesc)
+	{
+		// Convert FString to ANSI string and store them as intermediates
+		FMemory::Memzero(OutTargetDesc);
+
+		switch (InTarget.Language)
+		{
+		case EShaderConductorLanguage::Glsl:
+		case EShaderConductorLanguage::Essl:
+			ConvertScTargetDescLanguageGlslFamily(InTarget, OutTargetDesc);
+			break;
+		case EShaderConductorLanguage::Metal_macOS:
+		case EShaderConductorLanguage::Metal_iOS:
+			ConvertScTargetDescLanguageMetalFamily(InTarget, OutTargetDesc);
+			break;
+		}
+
+		// Convert flags map into an array container
+		TArray<ANSICHAR> FlagName, FlagValue;
+		for (const TPair<FString, FString>& Iter : InTarget.CompileFlags.GetDefinitionMap())
+		{
+			ConvertFStringToAnsiString(Iter.Key, FlagName);
+			ConvertFStringToAnsiString(Iter.Value, FlagValue);
+			Intermediates.Flags.Emplace(MoveTemp(FlagName), MoveTemp(FlagValue));
+		}
+
+		// Store references after all elements have been added to the container so the pointers remain valid
+		Intermediates.FlagRefs.SetNum(Intermediates.Flags.Num());
+		for (int32 Index = 0; Index < Intermediates.Flags.Num(); ++Index)
+		{
+			Intermediates.FlagRefs[Index].name = Intermediates.Flags[Index].Key.GetData();
+			Intermediates.FlagRefs[Index].value = Intermediates.Flags[Index].Value.GetData();
+		}
+
+		OutTargetDesc.options = Intermediates.FlagRefs.GetData();
+		OutTargetDesc.numOptions = static_cast<uint32>(Intermediates.FlagRefs.Num());
+
+		// Wrap input function into lambda to convert to ShaderConductor interface
+		if (InTarget.VariableTypeRenameCallback)
+		{
+			OutTargetDesc.variableTypeRenameCallback = [InnerCallback = InTarget.VariableTypeRenameCallback](const char* VariableName, const char* TypeName) -> ShaderConductor::Blob
+			{
+				// Forward callback to public interface callback
+				FString RenamedTypeName;
+				if (InnerCallback(FAnsiStringView(VariableName), FAnsiStringView(TypeName), RenamedTypeName))
+				{
+					if (!RenamedTypeName.IsEmpty())
+					{
+						// Convert renamed type name from FString to ShaderConductor::Blob
+						return ShaderConductor::Blob(TCHAR_TO_ANSI(*RenamedTypeName), RenamedTypeName.Len() + 1);
+					}
+				}
+				return ShaderConductor::Blob{};
+			};
+		}
+	}
+
+	static void ConvertScOptions(const FShaderConductorOptions& InOptions, ShaderConductor::Compiler::Options& OutOptions)
+	{
+		OutOptions.removeUnusedGlobals = InOptions.bRemoveUnusedGlobals;
+		OutOptions.packMatricesInRowMajor = InOptions.bPackMatricesInRowMajor;
+		OutOptions.enable16bitTypes = InOptions.bEnable16bitTypes;
+		OutOptions.enableDebugInfo = InOptions.bEnableDebugInfo;
+		OutOptions.disableOptimizations = InOptions.bDisableOptimizations;
+		OutOptions.enableFMAPass = InOptions.bEnableFMAPass;
+		if (OutOptions.enable16bitTypes)
+		{
+			// 16-bit types only supports with SM 6.2+
+			OutOptions.shaderModel = ShaderConductor::Compiler::ShaderModel{ 6, 2 };
+		}
+		else
+		{
+			OutOptions.shaderModel = ToShaderConductorShaderModel(InOptions.TargetProfile);
+		}
+	}
+
+	static void ConvertScDefines(FShaderConductorContext::FShaderConductorIntermediates& Intermediates, const FShaderCompilerDefinitions& InDefinitions)
+	{
+		// Convert FString to ANSI string for each macro definition and its value
+		TArray<ANSICHAR> DefineName, DefineValue;
+
+		for (const TPair<FString, FString>& Iter : InDefinitions.GetDefinitionMap())
+		{
+			ConvertFStringToAnsiString(Iter.Key, DefineName);
+			ConvertFStringToAnsiString(Iter.Value, DefineValue);
+			Intermediates.Defines.Emplace(MoveTemp(DefineName), MoveTemp(DefineValue));
+		}
+
+		// Store references after all elements have been added to the container so the pointers remain valid
+		Intermediates.DefineRefs.SetNum(Intermediates.Defines.Num());
+		for (int32 Index = 0; Index < Intermediates.Defines.Num(); ++Index)
+		{
+			Intermediates.DefineRefs[Index].name = Intermediates.Defines[Index].Key.GetData();
+			Intermediates.DefineRefs[Index].value = Intermediates.Defines[Index].Value.GetData();
+		}
+	}
+
+	// Returns whether the specified line of text contains only these characters, making it a valid line marker from DXC: ' ', '\t', '~', '^'
+	static bool IsTextLineDxcLineMarker(const FString& Line)
+	{
+		bool bContainsLineMarkerChars = false;
+		for (TCHAR Char : Line)
+		{
+			if (Char == TCHAR('~') || Char == TCHAR('^'))
+			{
+				// Line contains at least one of the necessary characters to be a potential DXC line marker.
+				bContainsLineMarkerChars = true;
+			}
+			else if (!(Char == TCHAR(' ') || Char == TCHAR('\t')))
+			{
+				// Illegal character for a potential DXC line marker.
+				return false;
+			}
+		}
+		return bContainsLineMarkerChars;
+	}
+
+	// Converts the error blob from ShaderConductor into an array of error reports (of type FShaderCompilerError).
+	static void ConvertScCompileErrors(ShaderConductor::Blob& ErrorBlob, TArray<FShaderCompilerError>& OutErrors)
+	{
+		// Convert blob into FString
+		FString ErrorString;
+		if (ConvertScBlobToFString(&ErrorBlob, ErrorString))
+		{
+			// Convert FString into array of FString (one for each line)
+			TArray<FString> ErrorStringLines;
+			ErrorString.ParseIntoArray(ErrorStringLines, TEXT("\n"));
+
+			// Forward parsed array of lines to primary conversion function
+			FShaderConductorContext::ConvertCompileErrors(MoveTemp(ErrorStringLines), OutErrors);
+		}
+	}
+
+	FShaderConductorContext::FShaderConductorContext()
+		: Intermediates(new FShaderConductorIntermediates())
+	{
+	}
+
+	FShaderConductorContext::~FShaderConductorContext()
+	{
+		delete Intermediates;
+	}
+
+	FShaderConductorContext::FShaderConductorContext(FShaderConductorContext&& Rhs)
+		: Errors(MoveTemp(Rhs.Errors))
+		, Intermediates(Rhs.Intermediates)
+	{
+		Rhs.Intermediates = nullptr;
+	}
+
+	FShaderConductorContext& FShaderConductorContext::operator = (FShaderConductorContext&& Rhs)
+	{
+		Errors = MoveTemp(Rhs.Errors);
+		delete Intermediates;
+		Intermediates = Rhs.Intermediates;
+		Rhs.Intermediates = nullptr;
+		return *this;
+	}
+
+	bool FShaderConductorContext::LoadSource(const FString& ShaderSource, const FString& Filename, const FString& EntryPoint, EHlslShaderFrequency ShaderStage, const FShaderCompilerDefinitions* Definitions)
+	{
+		// Convert FString to ANSI string and store them as intermediates
+		ConvertFStringToAnsiString(ShaderSource, Intermediates->ShaderSource);
+		ConvertFStringToAnsiString(Filename, Intermediates->Filename);
+		ConvertFStringToAnsiString(EntryPoint, Intermediates->EntryPoint);
+
+		// Convert macro definitions map into an array container
+		if (Definitions != nullptr)
+		{
+			ConvertScDefines(*Intermediates, *Definitions);
+		}
+
+		// Convert shader stage
+		Intermediates->Stage = ToShaderConductorShaderStage(ShaderStage);
+
+		return true;
+	}
+
+	bool FShaderConductorContext::LoadSource(const ANSICHAR* ShaderSource, const ANSICHAR* Filename, const ANSICHAR* EntryPoint, EHlslShaderFrequency ShaderStage, const FShaderCompilerDefinitions* Definitions)
+	{
+		// Store ANSI strings as intermediates
+		CopyAnsiString(ShaderSource, Intermediates->ShaderSource);
+		CopyAnsiString(Filename, Intermediates->Filename);
+		CopyAnsiString(EntryPoint, Intermediates->EntryPoint);
+
+		// Convert macro definitions map into an array container
+		if (Definitions != nullptr)
+		{
+			ConvertScDefines(*Intermediates, *Definitions);
+		}
+
+		// Convert shader stage
+		Intermediates->Stage = ToShaderConductorShaderStage(ShaderStage);
+
+		return true;
+	}
+
+	bool FShaderConductorContext::RewriteHlsl(const FShaderConductorOptions& Options, FString* OutSource)
+	{
+		// Convert descriptors for ShaderConductor interface
+		ShaderConductor::Compiler::SourceDesc ScSourceDesc;
+		ConvertScSourceDesc(*Intermediates, ScSourceDesc);
+
+		ShaderConductor::Compiler::Options ScOptions;
+		ConvertScOptions(Options, ScOptions);
+
+		// Rewrite HLSL with wrapper function to catch exceptions from ShaderConductor
+		bool bSucceeded = false;
+		bool bException = false;
+		ShaderConductor::Compiler::ResultDesc ResultDesc;
+		ScRewriteWrapper(ScSourceDesc, ScOptions, ResultDesc, bException);
+
+		if (!ResultDesc.hasError && !bException && ResultDesc.target.Size() > 0)
+		{
+			// Copy rewritten HLSL code into intermediate source code.
+			ConvertByteArrayToAnsiString(reinterpret_cast<const ANSICHAR*>(ResultDesc.target.Data()), ResultDesc.target.Size(), Intermediates->ShaderSource);
+
+			// If output source is specified, also convert to TCHAR string
+			if (OutSource != nullptr)
+			{
+				*OutSource = ANSI_TO_TCHAR(Intermediates->ShaderSource.GetData());
+			}
+			bSucceeded = true;
+		}
+		else
+		{
+			if (bException)
+			{
+				Errors.Add(TEXT("ShaderConductor exception during rewrite"));
+			}
+			bSucceeded = false;
+		}
+
+		// Append compile error and warning to output reports
+		ConvertScCompileErrors(ResultDesc.errorWarningMsg, Errors);
+
+		return bSucceeded;
+	}
+
+	bool FShaderConductorContext::CompileHlslToSpirv(const FShaderConductorOptions& Options, TArray<uint32>& OutSpirv)
+	{
+		// Convert descriptors for ShaderConductor interface
+		ShaderConductor::Compiler::SourceDesc ScSourceDesc;
+		ConvertScSourceDesc(*Intermediates, ScSourceDesc);
+
+		ShaderConductor::Compiler::TargetDesc ScTargetDesc;
+		FMemory::Memzero(ScTargetDesc);
+		ScTargetDesc.language = ShaderConductor::ShadingLanguage::SpirV;
+
+		ShaderConductor::Compiler::Options ScOptions;
+		ConvertScOptions(Options, ScOptions);
+
+		// Compile HLSL source code to SPIR-V
+		bool bSucceeded = false;
+		bool bException = false;
+		ShaderConductor::Compiler::ResultDesc ResultDesc;
+		ScCompileWrapper(ScSourceDesc, ScOptions, ScTargetDesc, ResultDesc, bException);
+
+		if (!ResultDesc.hasError && !bException && ResultDesc.target.Size() > 0)
+		{
+			// Copy result blob into output SPIR-V module
+			OutSpirv = TArray<uint32>(reinterpret_cast<const uint32*>(ResultDesc.target.Data()), ResultDesc.target.Size() / 4);
+			bSucceeded = true;
+		}
+		else
+		{
+			if (bException)
+			{
+				Errors.Add(TEXT("ShaderConductor exception during compilation"));
+			}
+			bSucceeded = false;
+		}
+
+		// Append compile error and warning to output reports
+		ConvertScCompileErrors(ResultDesc.errorWarningMsg, Errors);
+
+		return bSucceeded;
+	}
+
+	bool FShaderConductorContext::CompileSpirvToSource(const FShaderConductorOptions& Options, const FShaderConductorTarget& Target, const void* InSpirv, uint32 InSpirvByteSize, FString& OutSource)
+	{
+		return CompileSpirvToSourceBuffer(
+			Options, Target, InSpirv, InSpirvByteSize,
+			[&OutSource](const void* Data, uint32 Size)
+			{
+				// Convert source buffer to FString
+				ConvertByteArrayToFString(Data, Size, OutSource);
+			}
+		);
+	}
+
+	bool FShaderConductorContext::CompileSpirvToSourceAnsi(const FShaderConductorOptions& Options, const FShaderConductorTarget& Target, const void* InSpirv, uint32 InSpirvByteSize, TArray<ANSICHAR>& OutSource)
+	{
+		return CompileSpirvToSourceBuffer(
+			Options, Target, InSpirv, InSpirvByteSize,
+			[&OutSource](const void* Data, uint32 Size)
+			{
+				// Convert source buffer to ANSI string
+				ConvertByteArrayToAnsiString(reinterpret_cast<const ANSICHAR*>(Data), Size, OutSource);
+			}
+		);
+	}
+
+	bool FShaderConductorContext::CompileSpirvToSourceBuffer(const FShaderConductorOptions& Options, const FShaderConductorTarget& Target, const void* InSpirv, uint32 InSpirvByteSize, const TFunction<void(const void* Data, uint32 Size)>& OutputCallback)
+	{
+		check(OutputCallback != nullptr);
+		check(InSpirv != nullptr);
+		check(InSpirvByteSize > 0);
+		checkf(InSpirvByteSize % 4 == 0, TEXT("SPIR-V code unaligned. Size must be a multiple of 4, but %u was specified."), InSpirvByteSize);
+
+		// Convert descriptors for ShaderConductor interface
+		ShaderConductor::Compiler::SourceDesc ScSourceDesc;
+		ConvertScSourceDesc(*Intermediates, ScSourceDesc);
+
+		ShaderConductor::Compiler::TargetDesc ScTargetDesc;
+		ConvertScTargetDesc(*Intermediates, Target, ScTargetDesc);
+
+		ShaderConductor::Compiler::Options ScOptions;
+		ConvertScOptions(Options, ScOptions);
+
+		ShaderConductor::Compiler::ResultDesc ScBinaryDesc;
+		ScBinaryDesc.target.Reset(InSpirv, InSpirvByteSize);
+		ScBinaryDesc.isText = false;
+		ScBinaryDesc.hasError = false;
+
+		// Convert the input SPIR-V into Metal high level source
+		bool bSucceeded = false;
+		bool bException = false;
+		ShaderConductor::Compiler::ResultDesc ResultDesc;
+		ScConvertBinaryWrapper(ScBinaryDesc, ScSourceDesc, ScTargetDesc, ResultDesc, bException);
+
+		if (!ResultDesc.hasError && !bException && ResultDesc.target.Size() > 0)
+		{
+			// Copy result blob into output SPIR-V module
+			OutputCallback(ResultDesc.target.Data(), ResultDesc.target.Size());
+			bSucceeded = true;
+		}
+		else
+		{
+			if (bException)
+			{
+				Errors.Add(TEXT("ShaderConductor exception during SPIR-V binary conversion"));
+			}
+			bSucceeded = false;
+		}
+
+		// Append compile error and warning to output reports
+		if (ResultDesc.errorWarningMsg.Size() > 0)
+		{
+			FString ErrorString;
+			if (ConvertScBlobToFString(&ResultDesc.errorWarningMsg, ErrorString))
+			{
+				Errors.Add(*ErrorString);
+			}
+		}
+
+		return bSucceeded;
+	}
+
+	void FShaderConductorContext::FlushErrors(TArray<FShaderCompilerError>& OutErrors)
+	{
+		if (OutErrors.Num() > 0)
+		{
+			// Append internal list of errors to output list, then clear internal list
+			for (const FShaderCompilerError& ErrorEntry : Errors)
+			{
+				OutErrors.Add(ErrorEntry);
+			}
+			Errors.Empty();
+		}
+		else
+		{
+			// Move internal list of errors into output list
+			OutErrors = MoveTemp(Errors);
+		}
+	}
+
+	const ANSICHAR* FShaderConductorContext::GetSourceString() const
+	{
+		return (Intermediates->ShaderSource.Num() > 0 ? Intermediates->ShaderSource.GetData() : nullptr);
+	}
+
+	int32 FShaderConductorContext::GetSourceLength() const
+	{
+		return (Intermediates->ShaderSource.Num() > 0 ? (Intermediates->ShaderSource.Num() - 1) : 0);
+	}
+
+	void FShaderConductorContext::ConvertCompileErrors(TArray<FString>&& ErrorStringLines, TArray<FShaderCompilerError>& OutErrors)
+	{
+		// Returns whether the specified line in the 'ErrorStringLines' array has a line marker.
+		auto HasErrorLineMarker = [&ErrorStringLines](int32 LineIndex)
+		{
+			if (LineIndex + 2 < ErrorStringLines.Num())
+			{
+				return IsTextLineDxcLineMarker(ErrorStringLines[LineIndex + 2]);
+			}
+			return false;
+		};
+
+		// Iterate over all errors. Most (but not all) contain a highlighted line and line marker.
+		for (int32 LineIndex = 0; LineIndex < ErrorStringLines.Num();)
+		{
+			if (HasErrorLineMarker(LineIndex))
+			{
+				// Add current line as error with highlighted source line (LineIndex+1) and line marker (LineIndex+2)
+				OutErrors.Emplace(MoveTemp(ErrorStringLines[LineIndex]), MoveTemp(ErrorStringLines[LineIndex + 1]), MoveTemp(ErrorStringLines[LineIndex + 2]));
+				LineIndex += 3;
+			}
+			else
+			{
+				// Add current line as single error
+				OutErrors.Emplace(MoveTemp(ErrorStringLines[LineIndex]));
+				LineIndex += 1;
+			}
+		}
+	}
+
+#else // PLATFORM_MAC || PLATFORM_WINDOWS || PLATFORM_LINUX
+
+	FShaderConductorContext::FShaderConductorContext()
+	{
+		checkf(0, TEXT("Cannot instantiate FShaderConductorContext for unsupported platform"));
+	}
+
+	FShaderConductorContext::~FShaderConductorContext()
+	{
+		// Dummy
+	}
+
+	FShaderConductorContext::FShaderConductorContext(FShaderConductorContext&& Rhs)
+	{
+		// Dummy
+	}
+
+	FShaderConductorContext& FShaderConductorContext::operator = (FShaderConductorContext&& Rhs)
+	{
+		return *this; // Dummy
+	}
+
+	bool FShaderConductorContext::LoadSource(const FString& ShaderSource, const FString& Filename, const FString& EntryPoint, EHlslShaderFrequency ShaderStage, const FShaderCompilerDefinitions* Definitions)
+	{
+		return false; // Dummy
+	}
+
+	bool FShaderConductorContext::LoadSource(const ANSICHAR* ShaderSource, const ANSICHAR* Filename, const ANSICHAR* EntryPoint, EHlslShaderFrequency ShaderStage, const FShaderCompilerDefinitions* Definitions)
+	{
+		return false; // Dummy
+	}
+
+	bool FShaderConductorContext::RewriteHlsl(const FShaderConductorOptions& Options, FString* OutSource)
+	{
+		return false; // Dummy
+	}
+
+	bool FShaderConductorContext::CompileHlslToSpirv(const FShaderConductorOptions& Options, TArray<uint32>& OutSpirv)
+	{
+		return false; // Dummy
+	}
+
+	bool FShaderConductorContext::CompileSpirvToSource(const FShaderConductorOptions& Options, const FShaderConductorTarget& Target, const void* InSpirv, uint32 InSpirvByteSize, FString& OutSource)
+	{
+		return false; // Dummy
+	}
+
+	bool FShaderConductorContext::CompileSpirvToSourceAnsi(const FShaderConductorOptions& Options, const FShaderConductorTarget& Target, const void* InSpirv, uint32 InSpirvByteSize, TArray<ANSICHAR>& OutSource)
+	{
+		return false; // Dummy
+	}
+
+	bool FShaderConductorContext::CompileSpirvToSourceBuffer(const FShaderConductorOptions& Options, const FShaderConductorTarget& Target, const void* InSpirv, uint32 InSpirvByteSize, const TFunction<void(const void* Data, uint32 Size)>& OutputCallback)
+	{
+		return false; // Dummy
+	}
+
+	void FShaderConductorContext::FlushErrors(TArray<FShaderCompilerError>& OutErrors)
+	{
+		// Dummy
+	}
+
+	const ANSICHAR* FShaderConductorContext::GetSourceString() const
+	{
+		return nullptr; // Dummy
+	}
+
+	int32 FShaderConductorContext::GetSourceLength() const
+	{
+		return 0; // Dummy
+	}
+
+	void FShaderConductorContext::ConvertCompileErrors(const TArray<FString>& ErrorStringLines, TArray<FShaderCompilerError>& OutErrors)
+	{
+		// Dummy
+	}
+
+#endif // PLATFORM_MAC || PLATFORM_WINDOWS || PLATFORM_LINUX
+
+	bool FShaderConductorContext::IsIntermediateSpirvOutputVariable(const ANSICHAR* SpirvVariableName)
+	{
+		// This is only true for "temp.var.hullMainRetVal" which is generated by DXC as intermediate output variable to communicate patch constant data in a Hull Shader.
+		return (SpirvVariableName != nullptr && FCStringAnsi::Strcmp(SpirvVariableName, "temp.var.hullMainRetVal") == 0);
+	}
+
+} // namespace CrossCompiler
diff --git a/Engine/Source/Developer/ShaderCompilerCommon/Public/ShaderCompilerCommon.h b/Engine/Source/Developer/ShaderCompilerCommon/Public/ShaderCompilerCommon.h
--- a/Engine/Source/Developer/ShaderCompilerCommon/Public/ShaderCompilerCommon.h
+++ b/Engine/Source/Developer/ShaderCompilerCommon/Public/ShaderCompilerCommon.h
@@ -7,6 +7,7 @@
 #include "ShaderCore.h"
 #include "ShaderCompilerCore.h"
 #include "CrossCompilerDefinitions.h"
+#include "ShaderConductorContext.h"
 
 /**
  * This function looks for resources specified in ResourceTableMap in the 
@@ -381,131 +382,4 @@ namespace CrossCompiler
 		return false;
 	}
 
-	/** Wrapper structure to pass options descriptor to ShaderConductor. This is mapped to <struct ShaderConductor::Compiler::Options>. */
-	struct SHADERCOMPILERCOMMON_API FShaderConductorOptions
-	{
-		/** Removes unused global variables and resources. This can only be used in the HLSL rewrite pass, i.e. 'RewriteHlslSource'. */
-		bool bRemoveUnusedGlobals = false;
-
-		/** Experimental: Decide how a matrix get packed. Default in HLSL is row-major. This will be inverted in the SPIR-V backend to match SPIR-V's column-major default. */
-		bool bPackMatricesInRowMajor = true;
-
-		/** Enable 16-bit types, such as half, uint16_t. Requires shader model 6.2+. */
-		bool bEnable16bitTypes = false;
-
-		/** Embed debug info into the binary. */
-		bool bEnableDebugInfo = false;
-
-		/** Force to turn off optimizations. Ignore optimizationLevel below. */
-		bool bDisableOptimizations = false;
-
-		/** Enable a pass that converts floating point MUL+ADD pairs into FMAs to avoid re-association. */
-		bool bEnableFMAPass = false;
-
-		/** Target shader profile. By default HCT_FeatureLevelSM5. */
-		EHlslCompileTarget TargetProfile = HCT_FeatureLevelSM5;
-	};
-
-	/** Target high level languages for ShaderConductor output. */
-	enum class EShaderConductorLanguage
-	{
-		Glsl,
-		Essl,
-		Metal_macOS,
-		Metal_iOS,
-	};
-
-	/** Shader conductor output target descriptor. */
-	struct SHADERCOMPILERCOMMON_API FShaderConductorTarget
-	{
-		/** Target shader semantics, e.g. "macOS" or "iOS" for Metal GPU semantics. */
-		EShaderConductorLanguage Language = EShaderConductorLanguage::Glsl;
-
-		/**
-		Target shader version.
-		Valid values for Metal family: 20100, 20000, 10200, 10100, 10000.
-		Valid values for GLSL family: 310, 320, 330, 430.
-		*/
-		int32 Version = 0;
-
-		/** Cross compilation flags. This is used for high-level cross compilation (such as Metal output) that is send over to SPIRV-Cross, e.g. { "invariant_float_math", "1" }. */
-		FShaderCompilerDefinitions CompileFlags;
-
-		/** Optional callback to rename certain variable types. */
-		TFunction<bool(const FAnsiStringView& VariableName, const FAnsiStringView& TypeName, FString& OutRenamedTypeName)> VariableTypeRenameCallback;
-	};
-
-	/** Wrapper class to handle interface between UE and ShaderConductor. Use to compile HLSL shaders to SPIR-V or high-level languages such as Metal. */
-	class SHADERCOMPILERCOMMON_API FShaderConductorContext
-	{
-	public:
-		/** Initializes the context with internal buffers used for the conversion of input and option descriptors between UE and ShaderConductor. */
-		FShaderConductorContext();
-
-		/** Release the internal buffers. */
-		~FShaderConductorContext();
-
-		/** Move constructor to take ownership of internal buffers from 'Rhs'. */
-		FShaderConductorContext(FShaderConductorContext&& Rhs);
-
-		/** Move operator to take ownership of internal buffers from 'Rhs'. */
-		FShaderConductorContext& operator = (FShaderConductorContext&& Rhs);
-
-		FShaderConductorContext(const FShaderConductorContext&) = delete;
-		FShaderConductorContext& operator = (const FShaderConductorContext&) = delete;
-
-		/** Loads the shader source and converts the input descriptor to a format suitable for ShaderConductor. If 'Definitions' is null, the previously loaded definitions are not modified. */
-		bool LoadSource(const FString& ShaderSource, const FString& Filename, const FString& EntryPoint, EHlslShaderFrequency ShaderStage, const FShaderCompilerDefinitions* Definitions = nullptr);
-		bool LoadSource(const ANSICHAR* ShaderSource, const ANSICHAR* Filename, const ANSICHAR* EntryPoint, EHlslShaderFrequency ShaderStage, const FShaderCompilerDefinitions* Definitions = nullptr);
-
-		/** Rewrites the specified HLSL shader source code. This allows to reduce the HLSL code by removing unused global resources for instance.
-		This will update the internally loaded source (see 'LoadSource'), so the output parameter 'OutSource' is optional. */
-		bool RewriteHlsl(const FShaderConductorOptions& Options, FString* OutSource = nullptr);
-
-		/** Compiles the specified HLSL shader source code to SPIR-V. */
-		bool CompileHlslToSpirv(const FShaderConductorOptions& Options, TArray<uint32>& OutSpirv);
-
-		/** Compiles the specified SPIR-V shader binary code to high level source code (Metal or GLSL). */
-		bool CompileSpirvToSource(const FShaderConductorOptions& Options, const FShaderConductorTarget& Target, const void* InSpirv, uint32 InSpirvByteSize, FString& OutSource);
-
-		/** Compiles the specified SPIR-V shader binary code to high level source code (Metal or GLSL) stored as null terminated ANSI string. */
-		bool CompileSpirvToSourceAnsi(const FShaderConductorOptions& Options, const FShaderConductorTarget& Target, const void* InSpirv, uint32 InSpirvByteSize, TArray<ANSICHAR>& OutSource);
-
-		/** Compiles the specified SPIR-V shader binary code to high level source code (Metal or GLSL) stored as byte buffer (without null terminator as it comes from ShaderConductor). */
-		bool CompileSpirvToSourceBuffer(const FShaderConductorOptions& Options, const FShaderConductorTarget& Target, const void* InSpirv, uint32 InSpirvByteSize, const TFunction<void(const void* Data, uint32 Size)>& OutputCallback);
-
-		/** Flushes the list of current compile errors and moves the ownership to the caller. */
-		void FlushErrors(TArray<FShaderCompilerError>& OutErrors);
-
-		/** Returns a pointer to a null terminated ANSI string of the internal loaded sources, or null if no source has been loaded yet. This is automatically updated when RewriteHlsl() is called. */
-		const ANSICHAR* GetSourceString() const;
-
-		/** Returns a length of the internal loaded sources (excluding the null terminator). This is automatically updated when RewriteHlsl() is called. */
-		int32 GetSourceLength() const;
-
-		/** Returns the list of current compile errors. */
-		inline const TArray<FShaderCompilerError>& GetErrors() const
-		{
-			return Errors;
-		}
-
-	public:
-		/** Convert array of error string lines into array of <FShaderCompilerError>. */
-		static void ConvertCompileErrors(TArray<FString>&& ErrorStringLines, TArray<FShaderCompilerError>& OutErrors);
-
-		/** Returns whether the specified variable name denotes an intermediate output variable.
-		This is only true for a special identifiers generated by DXC to communicate patch constant data in the Hull Shader. */
-		static bool IsIntermediateSpirvOutputVariable(const ANSICHAR* SpirvVariableName);
-
-	public:
-		struct FShaderConductorIntermediates; // Pimpl idiom
-
-	private:
-		TArray<FShaderCompilerError> Errors;
-		FShaderConductorIntermediates* Intermediates; // Pimpl idiom
-	};
-
 }
-
-// Error code for SCW to help track down crashes
-extern SHADERCOMPILERCOMMON_API ESCWErrorCode GSCWErrorCode;
diff --git a/Engine/Source/Developer/ShaderCompilerCommon/Public/ShaderConductorContext.h b/Engine/Source/Developer/ShaderCompilerCommon/Public/ShaderConductorContext.h
new file mode 100644
--- /dev/null
+++ b/Engine/Source/Developer/ShaderCompilerCommon/Public/ShaderConductorContext.h
@@ -0,0 +1,141 @@
+// Copyright Epic Games, Inc. All Rights Reserved.
+
+#pragma once
+
+#include "CoreMinimal.h"
+#include "ShaderCore.h"
+#include "ShaderCompilerCore.h"
+#include "CrossCompilerDefinitions.h"
+
+// Cross compiler support/common functionality
+namespace CrossCompiler
+{
+
+	/** Wrapper structure to pass options descriptor to ShaderConductor. This is mapped to <struct ShaderConductor::Compiler::Options>. */
+	struct SHADERCOMPILERCOMMON_API FShaderConductorOptions
+	{
+		/** Removes unused global variables and resources. This can only be used in the HLSL rewrite pass, i.e. 'RewriteHlslSource'. */
+		bool bRemoveUnusedGlobals = false;
+
+		/** Experimental: Decide how a matrix get packed. Default in HLSL is row-major. This will be inverted in the SPIR-V backend to match SPIR-V's column-major default. */
+		bool bPackMatricesInRowMajor = true;
+
+		/** Enable 16-bit types, such as half, uint16_t. Requires shader model 6.2+. */
+		bool bEnable16bitTypes = false;
+
+		/** Embed debug info into the binary. */
+		bool bEnableDebugInfo = false;
+
+		/** Force to turn off optimizations. Ignore optimizationLevel below. */
+		bool bDisableOptimizations = false;
+
+		/** Enable a pass that converts floating point MUL+ADD pairs into FMAs to avoid re-association. */
+		bool bEnableFMAPass = false;
+
+		/** Target shader profile. By default HCT_FeatureLevelSM5. */
+		EHlslCompileTarget TargetProfile = HCT_FeatureLevelSM5;
+	};
+
+	/** Target high level languages for ShaderConductor output. */
+	enum class EShaderConductorLanguage
+	{
+		Glsl,
+		Essl,
+		Metal_macOS,
+		Metal_iOS,
+	};
+
+	/** Shader conductor output target descriptor. */
+	struct SHADERCOMPILERCOMMON_API FShaderConductorTarget
+	{
+		/** Target shader semantics, e.g. "macOS" or "iOS" for Metal GPU semantics. */
+		EShaderConductorLanguage Language = EShaderConductorLanguage::Glsl;
+
+		/**
+		Target shader version.
+		Valid values for Metal family: 20100, 20000, 10200, 10100, 10000.
+		Valid values for GLSL family: 310, 320, 330, 430.
+		*/
+		int32 Version = 0;
+
+		/** Cross compilation flags. This is used for high-level cross compilation (such as Metal output) that is send over to SPIRV-Cross, e.g. { "invariant_float_math", "1" }. */
+		FShaderCompilerDefinitions CompileFlags;
+
+		/** Optional callback to rename certain variable types. */
+		TFunction<bool(const FAnsiStringView& VariableName, const FAnsiStringView& TypeName, FString& OutRenamedTypeName)> VariableTypeRenameCallback;
+	};
+
+	/** Wrapper class to handle interface between UE and ShaderConductor. Use to compile HLSL shaders to SPIR-V or high-level languages such as Metal. */
+	class SHADERCOMPILERCOMMON_API FShaderConductorContext
+	{
+	public:
+		/** Initializes the context with internal buffers used for the conversion of input and option descriptors between UE and ShaderConductor. */
+		FShaderConductorContext();
+
+		/** Release the internal buffers. */
+		~FShaderConductorContext();
+
+		/** Move constructor to take ownership of internal buffers from 'Rhs'. */
+		FShaderConductorContext(FShaderConductorContext&& Rhs);
+
+		/** Move operator to take ownership of internal buffers from 'Rhs'. */
+		FShaderConductorContext& operator = (FShaderConductorContext&& Rhs);
+
+		FShaderConductorContext(const FShaderConductorContext&) = delete;
+		FShaderConductorContext& operator = (const FShaderConductorContext&) = delete;
+
+		/** Loads the shader source and converts the input descriptor to a format suitable for ShaderConductor. If 'Definitions' is null, the previously loaded definitions are not modified. */
+		bool LoadSource(const FString& ShaderSource, const FString& Filename, const FString& EntryPoint, EHlslShaderFrequency ShaderStage, const FShaderCompilerDefinitions* Definitions = nullptr);
+		bool LoadSource(const ANSICHAR* ShaderSource, const ANSICHAR* Filename, const ANSICHAR* EntryPoint, EHlslShaderFrequency ShaderStage, const FShaderCompilerDefinitions* Definitions = nullptr);
+
+		/** Rewrites the specified HLSL shader source code. This allows to reduce the HLSL code by removing unused global resources for instance.
+		This will update the internally loaded source (see 'LoadSource'), so the output parameter 'OutSource' is optional. */
+		bool RewriteHlsl(const FShaderConductorOptions& Options, FString* OutSource = nullptr);
+
+		/** Compiles the specified HLSL shader source code to SPIR-V. */
+		bool CompileHlslToSpirv(const FShaderConductorOptions& Options, TArray<uint32>& OutSpirv);
+
+		/** Compiles the specified SPIR-V shader binary code to high level source code (Metal or GLSL). */
+		bool CompileSpirvToSource(const FShaderConductorOptions& Options, const FShaderConductorTarget& Target, const void* InSpirv, uint32 InSpirvByteSize, FString& OutSource);
+
+		/** Compiles the specified SPIR-V shader binary code to high level source code (Metal or GLSL) stored as null terminated ANSI string. */
+		bool CompileSpirvToSourceAnsi(const FShaderConductorOptions& Options, const FShaderConductorTarget& Target, const void* InSpirv, uint32 InSpirvByteSize, TArray<ANSICHAR>& OutSource);
+
+		/** Compiles the specified SPIR-V shader binary code to high level source code (Metal or GLSL) stored as byte buffer (without null terminator as it comes from ShaderConductor). */
+		bool CompileSpirvToSourceBuffer(const FShaderConductorOptions& Options, const FShaderConductorTarget& Target, const void* InSpirv, uint32 InSpirvByteSize, const TFunction<void(const void* Data, uint32 Size)>& OutputCallback);
+
+		/** Flushes the list of current compile errors and moves the ownership to the caller. */
+		void FlushErrors(TArray<FShaderCompilerError>& OutErrors);
+
+		/** Returns a pointer to a null terminated ANSI string of the internal loaded sources, or null if no source has been loaded yet. This is automatically updated when RewriteHlsl() is called. */
+		const ANSICHAR* GetSourceString() const;
+
+		/** Returns a length of the internal loaded sources (excluding the null terminator). This is automatically updated when RewriteHlsl() is called. */
+		int32 GetSourceLength() const;
+
+		/** Returns the list of current compile errors. */
+		inline const TArray<FShaderCompilerError>& GetErrors() const
+		{
+			return Errors;
+		}
+
+	public:
+		/** Convert array of error string lines into array of <FShaderCompilerError>. */
+		static void ConvertCompileErrors(TArray<FString>&& ErrorStringLines, TArray<FShaderCompilerError>& OutErrors);
+
+		/** Returns whether the specified variable name denotes an intermediate output variable.
+		This is only true for a special identifiers generated by DXC to communicate patch constant data in the Hull Shader. */
+		static bool IsIntermediateSpirvOutputVariable(const ANSICHAR* SpirvVariableName);
+
+	public:
+		struct FShaderConductorIntermediates; // Pimpl idiom
+
+	private:
+		TArray<FShaderCompilerError> Errors;
+		FShaderConductorIntermediates* Intermediates; // Pimpl idiom
+	};
+
+}
+
+// Error code for SCW to help track down crashes
+extern SHADERCOMPILERCOMMON_API ESCWErrorCode GSCWErrorCode;
