Added HlslccHeaderWriter to unify writing of shader meta data.

#rnx

diff --git a/Engine/Source/Developer/Apple/MetalShaderFormat/Private/MetalDerivedData.cpp b/Engine/Source/Developer/Apple/MetalShaderFormat/Private/MetalDerivedData.cpp
--- a/Engine/Source/Developer/Apple/MetalShaderFormat/Private/MetalDerivedData.cpp
+++ b/Engine/Source/Developer/Apple/MetalShaderFormat/Private/MetalDerivedData.cpp
@@ -10,6 +10,7 @@
 #include "Misc/OutputDeviceRedirector.h"
 #include "MetalBackend.h"
 #include "Serialization/MemoryReader.h"
+#include "HlslccHeaderWriter.h"
 
 #include "MetalShaderFormat.h"
 
@@ -406,26 +407,6 @@ bool FMetalShaderOutputCooker::Build(TArray<uint8>& OutData)
 		// Load shader source into compiler context
 		CompilerContext.LoadSource(PreprocessedShader, Input.VirtualSourceFilePath, Input.EntryPointName, Frequency);
 
-		enum class EMetalTessellationMetadataTags : uint8
-		{
-			TessellationOutputControlPoints,
-			TessellationDomain,
-			TessellationInputControlPoints,
-			TessellationMaxTessFactor,
-			TessellationOutputWinding,
-			TessellationPartitioning,
-			TessellationPatchesPerThreadGroup,
-			TessellationPatchCountBuffer,
-			TessellationIndexBuffer,
-			TessellationHSOutBuffer,
-			TessellationControlPointOutBuffer,
-			TessellationHSTFOutBuffer,
-			TessellationControlPointIndexBuffer,
-			
-			Num
-		};
-		FString TESStrings[(uint8)EMetalTessellationMetadataTags::Num];
-        
 		if (bRewriteHlslSource)
 		{
 			// Rewrite HLSL source code to remove unused global resources and variables
@@ -457,20 +438,12 @@ bool FMetalShaderOutputCooker::Build(TArray<uint8>& OutData)
 			DumpDebugShaderText(Input, &SourceData[0], SourceData.size(), TEXT("rewritten.hlsl"));
 		}
 		
+		CrossCompiler::FHlslccHeaderWriter CCHeaderWriter;
+
 		FString MetaData = FString::Printf(TEXT("// ! %s/%s.usf:%s\n"), *Input.DebugGroupName, *Input.GetSourceFilename(), *Input.EntryPointName);
 		//FString EntryPoint = Input.EntryPointName;
 		EHlslShaderFrequency Freq = Frequency;
 		FString ALNString;
-		FString IABString;
-		FString UAVString;
-		FString SRVString;
-		FString SMPString;
-		FString UBOString;
-		FString GLOString;
-		FString PAKString;
-		FString INPString;
-		FString OUTString;
-		FString WKGString;
 		uint32 IABOffsetIndex = 0;
 		uint64 BufferIndices = 0xffffffffffffffff;
 
@@ -590,8 +563,6 @@ bool FMetalShaderOutputCooker::Build(TArray<uint8>& OutData)
 #endif
 				check(Found == 1);
 				
-				TESStrings[(uint8)EMetalTessellationMetadataTags::TessellationInputControlPoints] = FString::Printf(TEXT("// @TessellationInputControlPoints: %u\n"), PatchSize);
-				
 				ThreadSizeIdx = SourceData.find("[maxtessfactor(");
 				check(ThreadSizeIdx != std::string::npos);
 				String = SourceData.c_str() + ThreadSizeIdx;
@@ -604,9 +575,9 @@ bool FMetalShaderOutputCooker::Build(TArray<uint8>& OutData)
 #endif
 				check(Found == 1);
 				
-				TESStrings[(uint8)EMetalTessellationMetadataTags::TessellationMaxTessFactor] = FString::Printf(TEXT("// @TessellationMaxTessFactor: %u\n"), MaxTessFactor);
-				
-				TESStrings[(uint8)EMetalTessellationMetadataTags::TessellationPatchesPerThreadGroup] = TEXT("// @TessellationPatchesPerThreadGroup: 1\n");
+				CCHeaderWriter.WriteTessellationInputControlPoints(PatchSize);
+				CCHeaderWriter.WriteTessellationMaxTessFactor(MaxTessFactor);
+				CCHeaderWriter.WriteTessellationPatchesPerThreadGroup(1);
 			}
 			
 			{
@@ -625,23 +596,20 @@ bool FMetalShaderOutputCooker::Build(TArray<uint8>& OutData)
 							if (Frequency == HSF_ComputeShader)
 							{
 								check(Mode->operands_count == 3);
-								for (uint32 j = 0; j < Mode->operands_count; j++)
-								{
-									WKGString += FString::Printf(TEXT("%s%u"), WKGString.Len() ? TEXT(", ") : TEXT(""), Mode->operands[j]);
-								}
+								CCHeaderWriter.WriteNumThreads(Mode->operands[0], Mode->operands[1], Mode->operands[2]);
 							}
 							break;
 						case SpvExecutionModeTriangles:
 							if (Frequency == HSF_HullShader || Frequency == HSF_DomainShader)
 							{
-								TESStrings[(uint8)EMetalTessellationMetadataTags::TessellationDomain] = TEXT("// @TessellationDomain: tri\n");
+								CCHeaderWriter.WriteTessellationDomainTri();
 								Attribs.HSTFOutSize = 8; // sizeof(MTLTriangleTessellationFactorsHalf);
 							}
 							break;
 						case SpvExecutionModeQuads:
 							if (Frequency == HSF_HullShader || Frequency == HSF_DomainShader)
 							{
-								TESStrings[(uint8)EMetalTessellationMetadataTags::TessellationDomain] = TEXT("// @TessellationDomain: quad\n");
+								CCHeaderWriter.WriteTessellationDomainQuad();
 								Attribs.HSTFOutSize = 12; // sizeof(MTLQuadTessellationFactorsHalf);
 							}
 							break;
@@ -655,37 +623,37 @@ bool FMetalShaderOutputCooker::Build(TArray<uint8>& OutData)
 							if (Frequency == HSF_HullShader || Frequency == HSF_DomainShader)
 							{
 								check(Mode->operands_count == 1);
-								TESStrings[(uint8)EMetalTessellationMetadataTags::TessellationOutputControlPoints] = FString::Printf(TEXT("// @TessellationOutputControlPoints: %d\n"), Mode->operands[0]);
+								CCHeaderWriter.WriteTessellationOutputControlPoints(Mode->operands[0]);
 							}
 							break;
 						case SpvExecutionModeVertexOrderCw:
 							if (Frequency == HSF_HullShader || Frequency == HSF_DomainShader)
 							{
-								TESStrings[(uint8)EMetalTessellationMetadataTags::TessellationOutputWinding] = TEXT("// @TessellationOutputWinding: cw\n");
+								CCHeaderWriter.WriteTessellationOutputWindingCW();
 							}
 							break;
 						case SpvExecutionModeVertexOrderCcw:
 							if (Frequency == HSF_HullShader || Frequency == HSF_DomainShader)
 							{
-								TESStrings[(uint8)EMetalTessellationMetadataTags::TessellationOutputWinding] = TEXT("// @TessellationOutputWinding: ccw\n");
+								CCHeaderWriter.WriteTessellationOutputWindingCCW();
 							}
 							break;
 						case SpvExecutionModeSpacingEqual:
 							if (Frequency == HSF_HullShader || Frequency == HSF_DomainShader)
 							{
-								TESStrings[(uint8)EMetalTessellationMetadataTags::TessellationPartitioning] = TEXT("// @TessellationPartitioning: integer\n");
+								CCHeaderWriter.WriteTessellationPartitioningInteger();
 							}
 							break;
 						case SpvExecutionModeSpacingFractionalEven:
 							if (Frequency == HSF_HullShader || Frequency == HSF_DomainShader)
 							{
-								TESStrings[(uint8)EMetalTessellationMetadataTags::TessellationPartitioning] = TEXT("// @TessellationPartitioning: fractional_even\n");
+								CCHeaderWriter.WriteTessellationPartitioningFractionalEven();
 							}
 							break;
 						case SpvExecutionModeSpacingFractionalOdd:
 							if (Frequency == HSF_HullShader || Frequency == HSF_DomainShader)
 							{
-								TESStrings[(uint8)EMetalTessellationMetadataTags::TessellationPartitioning] = TEXT("// @TessellationPartitioning: fractional_odd\n");
+								CCHeaderWriter.WriteTessellationPartitioningFractionalOdd();
 							}
 							break;
 						case SpvExecutionModeInvocations:
@@ -933,7 +901,7 @@ bool FMetalShaderOutputCooker::Build(TArray<uint8>& OutData)
 					OutputData.TypedUAVs |= (1 << Index);
 					OutputData.TypedBuffers |= (1 << Index);
 					
-					UAVString += FString::Printf(TEXT("%s%s(%u:%u)"), UAVString.Len() ? TEXT(",") : TEXT(""), UTF8_TO_TCHAR(Binding->name), Index, 1);
+					CCHeaderWriter.WriteUAV(UTF8_TO_TCHAR(Binding->name), Index);
 					
 					SPVRResult = Reflection.ChangeDescriptorBindingNumbers(Binding, Index, GlobalSetId);
 					check(SPVRResult == SPV_REFLECT_RESULT_SUCCESS);
@@ -951,8 +919,8 @@ bool FMetalShaderOutputCooker::Build(TArray<uint8>& OutData)
 					
 					OutputData.InvariantBuffers |= (1 << Index);
 					
-					UAVString += FString::Printf(TEXT("%s%s(%u:%u)"), UAVString.Len() ? TEXT(",") : TEXT(""), UTF8_TO_TCHAR(Binding->name), Index, 1);
-					
+					CCHeaderWriter.WriteUAV(UTF8_TO_TCHAR(Binding->name), Index);
+
 					SPVRResult = Reflection.ChangeDescriptorBindingNumbers(Binding, Index, GlobalSetId);
 					check(SPVRResult == SPV_REFLECT_RESULT_SUCCESS);
 				}
@@ -968,8 +936,8 @@ bool FMetalShaderOutputCooker::Build(TArray<uint8>& OutData)
 					TextureIndices &= ~(1ull << (uint64)Index);
 					UAVIndices &= ~(1 << Index);
 					
-					UAVString += FString::Printf(TEXT("%s%s(%u:%u)"), UAVString.Len() ? TEXT(",") : TEXT(""), UTF8_TO_TCHAR(Binding->name), Index, 1);
-					
+					CCHeaderWriter.WriteUAV(UTF8_TO_TCHAR(Binding->name), Index);
+
 					SPVRResult = Reflection.ChangeDescriptorBindingNumbers(Binding, Index, GlobalSetId);
 					check(SPVRResult == SPV_REFLECT_RESULT_SUCCESS);
 				}
@@ -1072,27 +1040,20 @@ bool FMetalShaderOutputCooker::Build(TArray<uint8>& OutData)
 				{
 					FString Name = Pair.Key;
 					auto const& ResourceArray = Pair.Value;
-					if (IABString.Len())
-					{
-						IABString += TEXT(",");
-					}
 					uint32 SetIndex = ResourceArray[0].SetIndex + IABOffsetIndex;
-					IABString += FString::Printf(TEXT("%d["), SetIndex);
-					bool bComma = false;
+
+					TArray<uint32> IndirectArgumentBufferIndices;
+					IndirectArgumentBufferIndices.Reserve(ResourceArray.Num());
 					for (auto const& Resource : ResourceArray)
 					{
 						if (Resource.bUsed)
 						{
-							if (bComma)
-							{
-								IABString += TEXT(",");
-							}
-							IABString += FString::Printf(TEXT("%u"), (Resource.ResourceIndex == 65535 ? 0 : Resource.ResourceIndex + 1));
-							bComma = true;
+							IndirectArgumentBufferIndices.Add((Resource.ResourceIndex == 65535 ? 0 : Resource.ResourceIndex + 1));
 						}
 					}
-					IABString += TEXT("]");
-					UBOString += FString::Printf(TEXT("%s%s(%u)"), UBOString.Len() ? TEXT(",") : TEXT(""), *Name, SetIndex);
+
+					CCHeaderWriter.WriteArgumentBuffers(SetIndex, IndirectArgumentBufferIndices);
+					CCHeaderWriter.WriteUniformBlock(*Name, SetIndex);
 				}
 				
 				for (auto const& Binding : TBufferSRVBindings)
@@ -1106,7 +1067,7 @@ bool FMetalShaderOutputCooker::Build(TArray<uint8>& OutData)
 					
 					OutputData.TypedBuffers |= (1 << Index);
 					
-					SRVString += FString::Printf(TEXT("%s%s(%u:%u)"), SRVString.Len() ? TEXT(",") : TEXT(""), UTF8_TO_TCHAR(Binding->name), Index, 1);
+					CCHeaderWriter.WriteSRV(UTF8_TO_TCHAR(Binding->name), Index);
 					
 					SPVRResult = Reflection.ChangeDescriptorBindingNumbers(Binding, Index, GlobalSetId);
 					check(SPVRResult == SPV_REFLECT_RESULT_SUCCESS);
@@ -1121,8 +1082,8 @@ bool FMetalShaderOutputCooker::Build(TArray<uint8>& OutData)
 					
 					OutputData.InvariantBuffers |= (1 << Index);
 					
-					SRVString += FString::Printf(TEXT("%s%s(%u:%u)"), SRVString.Len() ? TEXT(",") : TEXT(""), UTF8_TO_TCHAR(Binding->name), Index, 1);
-					
+					CCHeaderWriter.WriteSRV(UTF8_TO_TCHAR(Binding->name), Index);
+
 					SPVRResult = Reflection.ChangeDescriptorBindingNumbers(Binding, Index, GlobalSetId);
 					check(SPVRResult == SPV_REFLECT_RESULT_SUCCESS);
 				}
@@ -1139,17 +1100,17 @@ bool FMetalShaderOutputCooker::Build(TArray<uint8>& OutData)
 					if (strstr(Binding->name, "$Globals"))
 					{
 						TCBDMARangeMap CBRanges;
-						GLOString = FString::Printf(TEXT("Globals(%u): "), Index);
+						CCHeaderWriter.WritePackedUB(Index);
 						
 						FString MbrString;
 						for (uint32 i = 0; i < Binding->block.member_count; i++)
 						{
 							SpvReflectBlockVariable& member = Binding->block.members[i];
-							uint32 MbrOffset = member.absolute_offset / sizeof(float);
-							uint32 MbrSize = member.size / sizeof(float);
 							
-							MbrString += FString::Printf(TEXT("%s%s(%u,%u)"), MbrString.Len() > 0 ? TEXT(",") : TEXT(""), UTF8_TO_TCHAR(member.name), MbrOffset, MbrSize);
+							CCHeaderWriter.WritePackedUBField(UTF8_TO_TCHAR(member.name), member.absolute_offset, member.size);
 							
+							const uint32 MbrOffset = member.absolute_offset / sizeof(float);
+							const uint32 MbrSize = member.size / sizeof(float);
 							unsigned DestCBPrecision = TEXT('h');
 							unsigned SourceOffset = MbrOffset;
 							unsigned DestOffset = MbrOffset;
@@ -1158,22 +1119,20 @@ bool FMetalShaderOutputCooker::Build(TArray<uint8>& OutData)
 							InsertRange(CBRanges, Index, SourceOffset, DestSize, DestCBIndex, DestCBPrecision, DestOffset);
 						}
 						
-						GLOString += MbrString;
-						
 						for (auto Iter = CBRanges.begin(); Iter != CBRanges.end(); ++Iter)
 						{
 							TDMARangeList& List = Iter->second;
 							for (auto IterList = List.begin(); IterList != List.end(); ++IterList)
 							{
 								check(IterList->DestCBIndex == 0);
-								PAKString += FString::Printf(TEXT("%s%u:%u-%c:%u:%u"), PAKString.Len() ? TEXT(",") : TEXT(""), IterList->SourceCB, IterList->SourceOffset, IterList->DestCBPrecision, IterList->DestOffset, IterList->Size);
+								CCHeaderWriter.WritePackedUBGlobalCopy(IterList->SourceCB, IterList->SourceOffset, IterList->DestCBIndex, IterList->DestCBPrecision, IterList->DestOffset, IterList->Size);
 							}
 						}
 					}
 					else
 					{
 						// Regular uniform buffer - we only care about the binding index
-						UBOString += FString::Printf(TEXT("%s%s(%u)"), UBOString.Len() ? TEXT(",") : TEXT(""), UTF8_TO_TCHAR(Binding->name), Index);
+						CCHeaderWriter.WriteUniformBlock(UTF8_TO_TCHAR(Binding->name), Index);
 					}
 					
 					SPVRResult = Reflection.ChangeDescriptorBindingNumbers(Binding, Index, GlobalSetId);
@@ -1186,7 +1145,7 @@ bool FMetalShaderOutputCooker::Build(TArray<uint8>& OutData)
 					uint32 Index = FPlatformMath::CountTrailingZeros64(TextureIndices);
 					TextureIndices &= ~(1ull << uint64(Index));
 					
-					SRVString += FString::Printf(TEXT("%s%s(%u:%u)"), SRVString.Len() ? TEXT(",") : TEXT(""), UTF8_TO_TCHAR(Binding->name), Index, 1);
+					CCHeaderWriter.WriteSRV(UTF8_TO_TCHAR(Binding->name), Index);
 					
 					SPVRResult = Reflection.ChangeDescriptorBindingNumbers(Binding, Index, GlobalSetId);
 					check(SPVRResult == SPV_REFLECT_RESULT_SUCCESS);
@@ -1198,7 +1157,7 @@ bool FMetalShaderOutputCooker::Build(TArray<uint8>& OutData)
 					uint32 Index = FPlatformMath::CountTrailingZeros64(SamplerIndices);
 					SamplerIndices &= ~(1ull << (uint64)Index);
 					
-					SMPString += FString::Printf(TEXT("%s%u:%s"), SMPString.Len() ? TEXT(",") : TEXT(""), Index, UTF8_TO_TCHAR(Binding->name));
+					CCHeaderWriter.WriteSamplerState(UTF8_TO_TCHAR(Binding->name), Index);
 					
 					SPVRResult = Reflection.ChangeDescriptorBindingNumbers(Binding, Index, GlobalSetId);
 					check(SPVRResult == SPV_REFLECT_RESULT_SUCCESS);
@@ -1244,7 +1203,7 @@ bool FMetalShaderOutputCooker::Build(TArray<uint8>& OutData)
 								TypeQualifier += TEXT("1");
 							}
 							
-							OUTString += FString::Printf(TEXT("%s%s:SV_Target%d"), OUTString.Len() ? TEXT(",") : TEXT(""), *TypeQualifier, Var->location);
+							CCHeaderWriter.WriteOutputAttribute(TEXT("SV_Target"), *TypeQualifier, Var->location, /*bLocationPrefix:*/ false, /*bLocationSuffix:*/ true);
 						}
 					}
 				}
@@ -1330,7 +1289,7 @@ bool FMetalShaderOutputCooker::Build(TArray<uint8>& OutData)
 							{
 								AssignedInputs |= (1 << (Location + j));
 								
-								INPString += FString::Printf(TEXT("%s%s:in_ATTRIBUTE%d"), INPString.Len() ? TEXT(",") : TEXT(""), *TypeQualifier, (Location + j));
+								CCHeaderWriter.WriteInputAttribute(TEXT("in_ATTRIBUTE"), *TypeQualifier, (Location + j), /*bLocationPrefix:*/ false, /*bLocationSuffix:*/ true);
 							}
 						}
 					}
@@ -1724,8 +1683,8 @@ bool FMetalShaderOutputCooker::Build(TArray<uint8>& OutData)
 				BufferIndices &= ~(1ull << (uint64)IndirectParamsIndex);
 				TargetDesc.CompileFlags.SetDefine(TEXT("indirect_params_buffer_index"), IndirectParamsIndex);
 				
-				TESStrings[(uint8)EMetalTessellationMetadataTags::TessellationHSOutBuffer] = FString::Printf(TEXT("// @TessellationHSOutBuffer: %u\n"), OutputBufferIndex);
-				TESStrings[(uint8)EMetalTessellationMetadataTags::TessellationPatchCountBuffer] = FString::Printf(TEXT("// @TessellationPatchCountBuffer: %u\n"), IndirectParamsIndex);
+				CCHeaderWriter.WriteTessellationHSOutBuffer(OutputBufferIndex);
+				CCHeaderWriter.WriteTessellationPatchCountBuffer(IndirectParamsIndex);
 			}
 			
 			if (Frequency == HSF_DomainShader)
@@ -1736,8 +1695,8 @@ bool FMetalShaderOutputCooker::Build(TArray<uint8>& OutData)
 				IndirectParamsIndex = FPlatformMath::CountTrailingZeros64(BufferIndices);
 				BufferIndices &= ~(1ull << (uint64)IndirectParamsIndex);
 				
-				TESStrings[(uint8)EMetalTessellationMetadataTags::TessellationHSOutBuffer] = FString::Printf(TEXT("// @TessellationHSOutBuffer: %u\n"), OutputBufferIndex);
-				TESStrings[(uint8)EMetalTessellationMetadataTags::TessellationControlPointOutBuffer] = FString::Printf(TEXT("// @TessellationControlPointOutBuffer: %u\n"), IndirectParamsIndex);
+				CCHeaderWriter.WriteTessellationHSOutBuffer(OutputBufferIndex);
+				CCHeaderWriter.WriteTessellationControlPointOutBuffer(IndirectParamsIndex);
 			}
 			
 			if (Frequency == HSF_HullShader)
@@ -1765,12 +1724,12 @@ bool FMetalShaderOutputCooker::Build(TArray<uint8>& OutData)
 				uint32 HullVertexBuffer = FPlatformMath::CountTrailingZeros64(BufferIndices);
 				BufferIndices &= ~(1ull << (uint64)HullVertexBuffer);
 				
-				TESStrings[(uint8)EMetalTessellationMetadataTags::TessellationHSOutBuffer] = FString::Printf(TEXT("// @TessellationHSOutBuffer: %u\n"), PatchBufferIndex);
-				TESStrings[(uint8)EMetalTessellationMetadataTags::TessellationPatchCountBuffer] = FString::Printf(TEXT("// @TessellationPatchCountBuffer: %u\n"), IndirectParamsIndex);
-				TESStrings[(uint8)EMetalTessellationMetadataTags::TessellationControlPointIndexBuffer] = FString::Printf(TEXT("// @TessellationControlPointIndexBuffer: %u\n"), HullVertexBuffer);
-				TESStrings[(uint8)EMetalTessellationMetadataTags::TessellationHSTFOutBuffer] = FString::Printf(TEXT("// @TessellationHSTFOutBuffer: %u\n"), TessFactorBufferIndex);
-				TESStrings[(uint8)EMetalTessellationMetadataTags::TessellationControlPointOutBuffer] = FString::Printf(TEXT("// @TessellationControlPointOutBuffer: %u\n"), OutputBufferIndex);
-				TESStrings[(uint8)EMetalTessellationMetadataTags::TessellationIndexBuffer] = FString::Printf(TEXT("// @TessellationIndexBuffer: %u\n"), HullIndexBuffer);
+				CCHeaderWriter.WriteTessellationHSOutBuffer(PatchBufferIndex);
+				CCHeaderWriter.WriteTessellationPatchCountBuffer(IndirectParamsIndex);
+				CCHeaderWriter.WriteTessellationControlPointIndexBuffer(HullVertexBuffer);
+				CCHeaderWriter.WriteTessellationHSTFOutBuffer(TessFactorBufferIndex);
+				CCHeaderWriter.WriteTessellationControlPointOutBuffer(OutputBufferIndex);
+				CCHeaderWriter.WriteTessellationIndexBuffer(HullIndexBuffer);
 			}
 			
 			switch (VersionEnum)
@@ -1835,62 +1794,13 @@ bool FMetalShaderOutputCooker::Build(TArray<uint8>& OutData)
 		}
 		else
 		{
-			MetaData += TEXT("// Compiled by ShaderConductor\n");
-			if (INPString.Len())
-			{
-				MetaData += FString::Printf(TEXT("// @Inputs: %s\n"), *INPString);
-			}
-			if (OUTString.Len())
+			if (FCStringAnsi::Strstr(ResultsTargetDataAsString.c_str(), "spvBufferSizeConstants"))
 			{
-				MetaData += FString::Printf(TEXT("// @Outputs: %s\n"), *OUTString);
-			}
-			if (UBOString.Len())
-			{
-				MetaData += FString::Printf(TEXT("// @UniformBlocks: %s\n"), *UBOString);
-			}
-			if (GLOString.Len())
-			{
-				MetaData += FString::Printf(TEXT("// @PackedUB: %s\n"), *GLOString);
-			}
-			if (PAKString.Len())
-			{
-				MetaData += FString::Printf(TEXT("// @PackedUBGlobalCopies: %s\n"), *PAKString);
-			}
-			if (SRVString.Len())
-			{
-				MetaData += FString::Printf(TEXT("// @Samplers: %s\n"), *SRVString);
-			}
-			if (UAVString.Len())
-			{
-				MetaData += FString::Printf(TEXT("// @UAVs: %s\n"), *UAVString);
-			}
-			if (SMPString.Len())
-			{
-				MetaData += FString::Printf(TEXT("// @SamplerStates: %s\n"), *SMPString);
-			}
-			if (WKGString.Len())
-			{
-				MetaData += FString::Printf(TEXT("// @NumThreads: %s\n"), *WKGString);
-			}
-			if(FCStringAnsi::Strstr(ResultsTargetDataAsString.c_str(), "spvBufferSizeConstants"))
-			{
-				MetaData += FString::Printf(TEXT("// @SideTable: spvBufferSizeConstants(%d)\n"), SideTableIndex);
-			}
-			if (IABString.Len())
-			{
-				MetaData += TEXT("// @ArgumentBuffers: ");
-				MetaData += IABString;
-				MetaData += TEXT("\n");
-			}
-			FString TESString;
-			for (uint32 i = 0; i < (uint32)EMetalTessellationMetadataTags::Num; i++)
-			{
-				TESString += TESStrings[i];
-			}
-			if (TESString.Len())
-			{
-				MetaData += TESString;
+				CCHeaderWriter.WriteSideTable(TEXT("spvBufferSizeConstants"), SideTableIndex);
 			}
+
+			MetaData += TEXT("// Compiled by ShaderConductor\n");
+			MetaData += CCHeaderWriter.ToString();
 			MetaData += TEXT("\n\n");
 			if (ALNString.Len())
 			{
diff --git a/Engine/Source/Developer/ShaderCompilerCommon/Private/HlslccHeaderWriter.cpp b/Engine/Source/Developer/ShaderCompilerCommon/Private/HlslccHeaderWriter.cpp
new file mode 100644
--- /dev/null
+++ b/Engine/Source/Developer/ShaderCompilerCommon/Private/HlslccHeaderWriter.cpp
@@ -0,0 +1,594 @@
+// Copyright Epic Games, Inc. All Rights Reserved.
+
+#include "HlslccHeaderWriter.h"
+#include "ShaderConductorContext.h"
+
+THIRD_PARTY_INCLUDES_START
+#include "spirv_reflect.h"
+THIRD_PARTY_INCLUDES_END
+
+
+namespace CrossCompiler
+{
+
+	template <typename FmtType, typename... Types>
+	static void MetaDataPrintf(FString& MetaData, const FmtType& Fmt, Types... Args)
+	{
+		if (!MetaData.IsEmpty())
+		{
+			MetaData += TEXT(",");
+		}
+		MetaData += FString::Printf(Fmt, Args...);
+	}
+
+	void FHlslccHeaderWriter::WriteInputAttribute(const SpvReflectInterfaceVariable& Attribute)
+	{
+		WriteIOAttribute(Strings.InputAttributes, Attribute, /*bIsInput:*/ true);
+	}
+
+	void FHlslccHeaderWriter::WriteInputAttribute(const TCHAR* AttributeName, const TCHAR* TypeSpecifier, int32 Location, bool bLocationPrefix, bool bLocationSuffix)
+	{
+		WriteIOAttribute(Strings.InputAttributes, AttributeName, TypeSpecifier, Location, bLocationPrefix, bLocationSuffix);
+	}
+
+	void FHlslccHeaderWriter::WriteOutputAttribute(const SpvReflectInterfaceVariable& Attribute)
+	{
+		WriteIOAttribute(Strings.OutputAttributes, Attribute, /*bIsInput:*/ false);
+	}
+
+	void FHlslccHeaderWriter::WriteOutputAttribute(const TCHAR* AttributeName, const TCHAR* TypeSpecifier, int32 Location, bool bLocationPrefix, bool bLocationSuffix)
+	{
+		WriteIOAttribute(Strings.OutputAttributes, AttributeName, TypeSpecifier, Location, bLocationPrefix, bLocationSuffix);
+	}
+
+	static void ConvertMetaDataTypeSpecifierPrimary(const SpvReflectTypeDescription& TypeSpecifier, FString& OutTypeName, uint32& OutTypeBitWidth, bool bBaseTypeOnly)
+	{
+		// Generate prefix for base type
+		if (TypeSpecifier.type_flags & SPV_REFLECT_TYPE_FLAG_BOOL)
+		{
+			OutTypeName += TEXT('b');
+			OutTypeBitWidth = 8;
+		}
+		else if (TypeSpecifier.type_flags & SPV_REFLECT_TYPE_FLAG_INT)
+		{
+			if (TypeSpecifier.traits.numeric.scalar.signedness)
+			{
+				OutTypeName += TEXT('i');
+			}
+			else
+			{
+				OutTypeName += TEXT('u');
+			}
+			OutTypeBitWidth = 32;
+		}
+		else if (TypeSpecifier.type_flags & SPV_REFLECT_TYPE_FLAG_FLOAT)
+		{
+			if (TypeSpecifier.traits.numeric.scalar.width == 16)
+			{
+				OutTypeName += TEXT('h');
+				OutTypeBitWidth = 16;
+			}
+			else
+			{
+				OutTypeName += TEXT('f');
+				OutTypeBitWidth = 32;
+			}
+		}
+
+		if (!bBaseTypeOnly)
+		{
+			// Generate number for vector size
+			const SpvReflectTypeFlags SpvScalarTypeFlags = (SPV_REFLECT_TYPE_FLAG_BOOL | SPV_REFLECT_TYPE_FLAG_INT | SPV_REFLECT_TYPE_FLAG_FLOAT);
+			if (TypeSpecifier.type_flags & SPV_REFLECT_TYPE_FLAG_VECTOR)
+			{
+				static const TCHAR* VectorDims = TEXT("1234");
+				const uint32 VectorSize = TypeSpecifier.traits.numeric.vector.component_count;
+				check(VectorSize >= 1 && VectorSize <= 4);
+				OutTypeName += VectorDims[VectorSize - 1];
+			}
+			else if (TypeSpecifier.type_flags & SPV_REFLECT_TYPE_FLAG_MATRIX)
+			{
+				//TODO
+			}
+			else if ((TypeSpecifier.type_flags & SpvScalarTypeFlags) != 0)
+			{
+				OutTypeName += TEXT('1'); // add single scalar component
+			}
+		}
+	}
+
+	static FString ConvertMetaDataTypeSpecifier(const SpvReflectTypeDescription& TypeSpecifier, uint32* OutTypeBitWidth = nullptr, bool bBaseTypeOnly = false)
+	{
+		FString TypeName;
+		uint32 TypeBitWidth = sizeof(float) * 8;
+		ConvertMetaDataTypeSpecifierPrimary(TypeSpecifier, TypeName, TypeBitWidth, bBaseTypeOnly);
+		if (OutTypeBitWidth)
+		{
+			*OutTypeBitWidth = TypeBitWidth;
+		}
+		return TypeName;
+	}
+
+	static const TCHAR* SpvBuiltinToString(const SpvBuiltIn BuiltIn)
+	{
+		switch (BuiltIn)
+		{
+		case SpvBuiltInPosition:					return TEXT("gl_Position");
+		case SpvBuiltInPointSize:					return TEXT("gl_PointSize");
+		case SpvBuiltInClipDistance:				return TEXT("gl_ClipDistance");
+		case SpvBuiltInCullDistance:				return TEXT("gl_CullDistance");
+		case SpvBuiltInVertexId:					return TEXT("gl_VertexID");
+		case SpvBuiltInInstanceId:					return TEXT("gl_InstanceID");
+		case SpvBuiltInPrimitiveId:					return TEXT("gl_PrimitiveID");
+		case SpvBuiltInInvocationId:				return TEXT("gl_InvocationID");
+		case SpvBuiltInLayer:						return TEXT("gl_Layer");
+		case SpvBuiltInViewportIndex:				return TEXT("gl_ViewportIndex");
+		case SpvBuiltInTessLevelOuter:				return TEXT("gl_TessLevelOuter");
+		case SpvBuiltInTessLevelInner:				return TEXT("gl_TessLevelInner");
+		case SpvBuiltInTessCoord:					return TEXT("gl_TessCoord");
+		case SpvBuiltInPatchVertices:				return TEXT("gl_PatchVertices");
+		case SpvBuiltInFragCoord:					return TEXT("gl_FragCoord");
+		case SpvBuiltInPointCoord:					return TEXT("gl_PointCoord");
+		case SpvBuiltInFrontFacing:					return TEXT("gl_FrontFacing");
+		case SpvBuiltInSampleId:					return TEXT("gl_SampleID");
+		case SpvBuiltInSamplePosition:				return TEXT("gl_SamplePosition");
+		case SpvBuiltInSampleMask:					return TEXT("gl_SampleMask");
+		case SpvBuiltInFragDepth:					return TEXT("gl_FragDepth");
+		case SpvBuiltInHelperInvocation:			return TEXT("gl_HelperInvocation");
+		case SpvBuiltInNumWorkgroups:				return TEXT("gl_NumWorkgroups");
+		case SpvBuiltInWorkgroupSize:				return TEXT("gl_WorkgroupSize");
+		case SpvBuiltInWorkgroupId:					return TEXT("gl_WorkgroupID");
+		case SpvBuiltInLocalInvocationId:			return TEXT("gl_LocalInvocationID");
+		case SpvBuiltInGlobalInvocationId:			return TEXT("gl_GlobalInvocationID");
+		case SpvBuiltInLocalInvocationIndex:		return TEXT("gl_LocalInvocationIndex");
+		case SpvBuiltInWorkDim:						return TEXT("gl_WorkDim");
+		case SpvBuiltInGlobalSize:					return TEXT("gl_GlobalSize");
+		case SpvBuiltInEnqueuedWorkgroupSize:		return TEXT("gl_EnqueuedWorkgroupSize");
+		case SpvBuiltInGlobalOffset:				return TEXT("gl_GlobalOffset");
+		case SpvBuiltInGlobalLinearId:				return TEXT("gl_GlobalLinearID");
+		case SpvBuiltInSubgroupSize:				return TEXT("gl_SubgroupSize");
+		case SpvBuiltInSubgroupMaxSize:				return TEXT("gl_SubgroupMaxSize");
+		case SpvBuiltInNumSubgroups:				return TEXT("gl_NumSubgroups");
+		case SpvBuiltInNumEnqueuedSubgroups:		return TEXT("gl_NumEnqueuedSubgroups");
+		case SpvBuiltInSubgroupId:					return TEXT("gl_SubgroupID");
+		case SpvBuiltInSubgroupLocalInvocationId:	return TEXT("gl_SubgroupLocalInvocationID");
+		case SpvBuiltInVertexIndex:					return TEXT("gl_VertexIndex");
+		case SpvBuiltInInstanceIndex:				return TEXT("gl_InstanceIndex");
+		case SpvBuiltInSubgroupEqMask:				return TEXT("gl_SubgroupEqMask");
+		case SpvBuiltInSubgroupGeMask:				return TEXT("gl_SubgroupGeMask");
+		case SpvBuiltInSubgroupGtMask:				return TEXT("gl_SubgroupGtMask");
+		case SpvBuiltInSubgroupLeMask:				return TEXT("gl_SubgroupLeMask");
+		case SpvBuiltInSubgroupLtMask:				return TEXT("gl_SubgroupLtMask");
+		case SpvBuiltInBaseVertex:					return TEXT("gl_BaseVertex");
+		case SpvBuiltInBaseInstance:				return TEXT("gl_BaseInstance");
+		case SpvBuiltInDrawIndex:					return TEXT("gl_DrawIndex");
+		case SpvBuiltInDeviceIndex:					return TEXT("gl_DeviceIndex");
+		case SpvBuiltInViewIndex:					return TEXT("gl_ViewIndex");
+		}
+		return nullptr;
+	}
+
+	static FString ConvertAttributeToMetaDataSemantic(const ANSICHAR* AttributeName, const SpvBuiltIn BuiltIn, bool bIsInput)
+	{
+		if (const TCHAR* BuiltInName = SpvBuiltinToString(BuiltIn))
+		{
+			return FString(BuiltInName);
+		}
+		else
+		{
+			check(AttributeName != nullptr && *AttributeName != '\0');
+			FString InSemantic = ANSI_TO_TCHAR(AttributeName);
+			FString OutSemantic = (bIsInput ? TEXT("in_") : TEXT("out_"));
+
+			if (InSemantic.StartsWith(TEXT("SV_")))
+			{
+				OutSemantic += InSemantic.Right(InSemantic.Len() - 3);
+			}
+			else
+			{
+				OutSemantic += InSemantic;
+			}
+
+			return OutSemantic;
+		}
+	}
+
+	// Flattens the array dimensions of the interface variable (aka shader attribute), e.g. from float4[2][3] -> float4[6]
+	static uint32 FlattenAttributeArrayDimension(const SpvReflectInterfaceVariable& Attribute, uint32 FirstArrayDim = 0)
+	{
+		uint32 FlattenedArrayDim = 1;
+		for (uint32 ArrayDimIndex = FirstArrayDim; ArrayDimIndex < Attribute.array.dims_count; ++ArrayDimIndex)
+		{
+			FlattenedArrayDim *= Attribute.array.dims[ArrayDimIndex];
+		}
+		return FlattenedArrayDim;
+	}
+
+	// Returns the string position where the index in the specified HLSL semantic beings, e.g. "SV_Target2" -> 9, "SV_Target" -> INDEX_NONE
+	static int32 FindIndexInHlslSemantic(const FString& Semantic)
+	{
+		int32 Index = Semantic.Len();
+		if (Index > 0 && FChar::IsDigit(Semantic[Index - 1]))
+		{
+			while (Index > 0 && FChar::IsDigit(Semantic[Index - 1]))
+			{
+				--Index;
+			}
+			return Index;
+		}
+		return INDEX_NONE;
+	}
+
+	// private
+	void FHlslccHeaderWriter::WriteIOAttribute(FString& OutMetaData, const TCHAR* AttributeName, const TCHAR* TypeSpecifier, int32 Location, bool bLocationPrefix, bool bLocationSuffix)
+	{
+		MetaDataPrintf(OutMetaData, TEXT("%s"), TypeSpecifier);
+		if (bLocationPrefix)
+		{
+			OutMetaData += FString::Printf(TEXT(";%d:"), Location);
+		}
+		else
+		{
+			OutMetaData += TEXT(":");
+		}
+		if (bLocationSuffix)
+		{
+			OutMetaData += FString::Printf(TEXT("%s%d"), AttributeName, Location);
+		}
+		else
+		{
+			OutMetaData += AttributeName;
+		}
+	}
+
+	// private
+	void FHlslccHeaderWriter::WriteIOAttribute(FString& OutMetaData, const SpvReflectInterfaceVariable& Attribute, bool bIsInput)
+	{
+		// Ignore interface variables that are only generated for intermediate results
+		if (CrossCompiler::FShaderConductorContext::IsIntermediateSpirvOutputVariable(Attribute.name))
+		{
+			return;
+		}
+
+		const FString TypeSpecifier = ConvertMetaDataTypeSpecifier(*Attribute.type_description);
+		FString Semantic = ConvertAttributeToMetaDataSemantic(Attribute.semantic, Attribute.built_in, bIsInput);
+
+		if (Attribute.array.dims_count > 0)
+		{
+			// Get semantic without index, e.g. "out_Target0" -> "out_Target"
+			const int32 SemanticIndexPos = FindIndexInHlslSemantic(Semantic);
+			if (SemanticIndexPos != INDEX_NONE)
+			{
+				Semantic = Semantic.Left(SemanticIndexPos);
+			}
+
+			if (Attribute.location == -1)
+			{
+				// Flatten array dimensions, e.g. from float4[3][2] -> float4[6]
+				const uint32 FlattenedArrayDim = FlattenAttributeArrayDimension(Attribute);
+
+				// Emit one output slot for each array element, e.g. "out float4 OutColor[2] : SV_Target0" occupies output slot SV_Target0 and SV_Target1.
+				for (uint32 FlattenedArrayIndex = 0; FlattenedArrayIndex < FlattenedArrayDim; ++FlattenedArrayIndex)
+				{
+					// If there is no binding slot, emit output as system value array such as "gl_SampleMask[]"
+					MetaDataPrintf(
+						OutMetaData,
+						TEXT("%s;%d:%s[%d]"),
+						*TypeSpecifier, // type specifier
+						Attribute.location,
+						*Semantic,
+						FlattenedArrayIndex
+					);
+				}
+			}
+			else if (!bIsInput)
+			{
+				//NOTE: For some reason, the meta data for output slot arrays must be entirely flattened, including the outer most array dimension
+				// Flatten array dimensions, e.g. from float4[3][2] -> float4[6]
+				const uint32 FlattenedArrayDim = FlattenAttributeArrayDimension(Attribute);
+
+				// Emit one output slot for each array element, e.g. "out float4 OutColor[2] : SV_Target0" occupies output slot SV_Target0 and SV_Target1.
+				for (uint32 FlattenedArrayIndex = 0; FlattenedArrayIndex < FlattenedArrayDim; ++FlattenedArrayIndex)
+				{
+					const uint32 BindingSlot = Attribute.location + FlattenedArrayIndex;
+					MetaDataPrintf(
+						OutMetaData,
+						TEXT("%s;%d:%s%d"),
+						*TypeSpecifier, // Type specifier
+						BindingSlot,
+						*Semantic,
+						BindingSlot
+					);
+				}
+			}
+			else if (Attribute.array.dims_count >= 2)
+			{
+				// Flatten array dimensions, e.g. from float4[3][2] -> float4[6]
+				const uint32 FlattenedArrayDim = FlattenAttributeArrayDimension(Attribute, 1);
+
+				// Emit one output slot for each array element, e.g. "out float4 OutColor[2] : SV_Target0" occupies output slot SV_Target0 and SV_Target1.
+				for (uint32 FlattenedArrayIndex = 0; FlattenedArrayIndex < FlattenedArrayDim; ++FlattenedArrayIndex)
+				{
+					const uint32 BindingSlot = Attribute.location + FlattenedArrayIndex;
+					MetaDataPrintf(
+						OutMetaData,
+						TEXT("%s[%d];%d:%s%d"),
+						*TypeSpecifier, // Type specifier
+						Attribute.array.dims[0], // Outer most array dimension
+						BindingSlot,
+						*Semantic,
+						BindingSlot
+					);
+				}
+			}
+			else
+			{
+				const uint32 BindingSlot = Attribute.location;
+				MetaDataPrintf(
+					OutMetaData,
+					TEXT("%s[%d];%d:%s%d"),
+					*TypeSpecifier, // Type specifier
+					Attribute.array.dims[0], // Outer most array dimension
+					BindingSlot,
+					*Semantic,
+					BindingSlot
+				);
+			}
+		}
+		else
+		{
+			MetaDataPrintf(
+				OutMetaData,
+				TEXT("%s;%d:%s"),
+				*TypeSpecifier, // type specifier
+				Attribute.location,
+				*Semantic
+			);
+		}
+	}
+
+	void FHlslccHeaderWriter::WriteUniformBlock(const TCHAR* ResourceName, uint32 BindingIndex)
+	{
+		MetaDataPrintf(Strings.UniformBlocks, TEXT("%s(%u)"), ResourceName, BindingIndex);
+	}
+
+	void FHlslccHeaderWriter::WritePackedGlobal(const TCHAR* ResourceName, const TCHAR* TypeSpecifier, uint32 ByteOffset, uint32 ByteSize)
+	{
+		checkf(ByteOffset % 4 == 0, TEXT("field offset in @PackedGlobals shader meta data must be a multiple of 4, but got %u"), ByteOffset);
+		checkf(ByteSize % 4 == 0, TEXT("field size in @PackedGlobals shader meta data must be a multiple of 4, but got %u"), ByteSize);
+		MetaDataPrintf(Strings.PackedGlobals, TEXT("%s(%s:%u,%u)"), ResourceName, TypeSpecifier, ByteOffset / 4, ByteSize / 4);
+	}
+
+	void FHlslccHeaderWriter::WritePackedUB(uint32 BindingIndex)
+	{
+		checkf(Strings.PackedUB.IsEmpty(), TEXT("cannot define @PackedUB attribute more than once"));
+		MetaDataPrintf(Strings.PackedUB, TEXT("Globals(%u): "), BindingIndex);
+	}
+
+	void FHlslccHeaderWriter::WritePackedUBField(const TCHAR* ResourceName, uint32 ByteOffset, uint32 ByteSize)
+	{
+		checkf(!Strings.PackedUB.IsEmpty(), TEXT("cannot append field without @PackedUB attribute in shader meta data"));
+		checkf(ByteOffset % 4 == 0, TEXT("field offset in @PackedUB shader meta data must be a multiple of 4, but got %u"), ByteOffset);
+		checkf(ByteSize % 4 == 0, TEXT("field size in @PackedUB shader meta data must be a multiple of 4, but got %u"), ByteSize);
+		MetaDataPrintf(Strings.PackedUBFields, TEXT("%s(%u,%u)"), ResourceName, ByteOffset / 4, ByteSize / 4);
+	}
+
+	void FHlslccHeaderWriter::WritePackedUBGlobalCopy(uint32 SourceCB, uint32 SourceOffset, uint32 DestCBIndex, uint32 DestCBPrecision, uint32 DestOffset, uint32 Size, bool bGroupFlattenedUBs)
+	{
+		if (bGroupFlattenedUBs)
+		{
+			MetaDataPrintf(Strings.PackedUBGlobalCopies, TEXT("%u:%u-%u:%c:%u:%u"), SourceCB, SourceOffset, DestCBIndex, DestCBPrecision, DestOffset, Size);
+		}
+		else
+		{
+			check(DestCBIndex == 0);
+			MetaDataPrintf(Strings.PackedUBGlobalCopies, TEXT("%u:%u-%c:%u:%u"), SourceCB, SourceOffset, DestCBPrecision, DestOffset, Size);
+		}
+	}
+
+	void FHlslccHeaderWriter::WriteSRV(const TCHAR* ResourceName, uint32 BindingIndex, uint32 Count)
+	{
+		MetaDataPrintf(Strings.SRVs, TEXT("%s(%u:%u)"), ResourceName, BindingIndex, Count);
+	}
+
+	void FHlslccHeaderWriter::WriteSRV(const TCHAR* ResourceName, uint32 BindingIndex, uint32 Count, const TArray<FString>& AssociatedResourceNames)
+	{
+		MetaDataPrintf(Strings.SRVs, TEXT("%s(%u:%u"), ResourceName, BindingIndex, Count);
+
+		if (!AssociatedResourceNames.IsEmpty())
+		{
+			Strings.SRVs += TEXT("[");
+			for (int32 ArrayIndex = 0; ArrayIndex < AssociatedResourceNames.Num(); ++ArrayIndex)
+			{
+				if (ArrayIndex > 0)
+				{
+					Strings.SRVs += TEXT(",");
+				}
+				Strings.SRVs += AssociatedResourceNames[ArrayIndex];
+			}
+			Strings.SRVs += TEXT("]");
+		}
+
+		Strings.SRVs += TEXT(")");
+	}
+
+	void FHlslccHeaderWriter::WriteUAV(const TCHAR* ResourceName, uint32 BindingIndex, uint32 Count)
+	{
+		MetaDataPrintf(Strings.UAVs, TEXT("%s(%u:%u)"), ResourceName, BindingIndex, Count);
+	}
+
+	void FHlslccHeaderWriter::WriteSamplerState(const TCHAR* ResourceName, uint32 BindingIndex)
+	{
+		MetaDataPrintf(Strings.SamplerStates, TEXT("%u:%s"), BindingIndex, ResourceName);
+	}
+
+	void FHlslccHeaderWriter::WriteNumThreads(uint32 NumThreadsX, uint32 NumThreadsY, uint32 NumThreadsZ)
+	{
+		MetaDataPrintf(Strings.NumThreads, TEXT("%u, %u, %u"), NumThreadsX, NumThreadsY, NumThreadsZ);
+	}
+
+	void FHlslccHeaderWriter::WriteSideTable(const TCHAR* ResourceName, uint32 SideTableIndex)
+	{
+		MetaDataPrintf(Strings.SideTable, TEXT("%s(%d)"), ResourceName, SideTableIndex);
+	}
+
+	void FHlslccHeaderWriter::WriteArgumentBuffers(uint32 BindingIndex, const TArray<uint32>& ResourceIndices)
+	{
+		MetaDataPrintf(Strings.ArgumentBuffers, TEXT("%d["), BindingIndex);
+
+		for (int32 ArrayIndex = 0; ArrayIndex < ResourceIndices.Num(); ++ArrayIndex)
+		{
+			if (ArrayIndex > 0)
+			{
+				Strings.ArgumentBuffers += TEXT(",");
+			}
+			Strings.ArgumentBuffers += FString::Printf(TEXT("%u"), ResourceIndices[ArrayIndex]);
+		}
+
+		Strings.ArgumentBuffers += TEXT("]");
+	}
+
+	void FHlslccHeaderWriter::WriteTessellationInputControlPoints(uint32 PatchSize)
+	{
+		checkf(Strings.TessellationInputControlPoints.IsEmpty(), TEXT("cannot define @TessellationInputControlPoints attribute more than once"));
+		Strings.TessellationInputControlPoints = FString::Printf(TEXT("%u"), PatchSize);
+	}
+
+	void FHlslccHeaderWriter::WriteTessellationOutputControlPoints(uint32 PatchSize)
+	{
+		checkf(Strings.TessellationOutputControlPoints.IsEmpty(), TEXT("cannot define @TessellationOutputControlPoints attribute more than once"));
+		Strings.TessellationOutputControlPoints = FString::Printf(TEXT("%u"), PatchSize);
+	}
+
+	void FHlslccHeaderWriter::WriteTessellationMaxTessFactor(uint32 MaxTessFactor)
+	{
+		checkf(Strings.TessellationMaxTessFactor.IsEmpty(), TEXT("cannot define @TessellationMaxTessFactor attribute more than once"));
+		Strings.TessellationMaxTessFactor = FString::Printf(TEXT("%u"), MaxTessFactor);
+	}
+
+	void FHlslccHeaderWriter::WriteTessellationDomainTri()
+	{
+		checkf(Strings.TessellationDomain.IsEmpty(), TEXT("cannot define @TessellationDomain attribute more than once"));
+		Strings.TessellationDomain = TEXT("tri");
+	}
+
+	void FHlslccHeaderWriter::WriteTessellationDomainQuad()
+	{
+		checkf(Strings.TessellationDomain.IsEmpty(), TEXT("cannot define @TessellationDomain attribute more than once"));
+		Strings.TessellationDomain = TEXT("quad");
+	}
+
+	void FHlslccHeaderWriter::WriteTessellationOutputWindingCW()
+	{
+		checkf(Strings.TessellationOutputWinding.IsEmpty(), TEXT("cannot define @TessellationOutputWinding attribute more than once"));
+		Strings.TessellationOutputWinding = TEXT("cw");
+	}
+
+	void FHlslccHeaderWriter::WriteTessellationOutputWindingCCW()
+	{
+		checkf(Strings.TessellationOutputWinding.IsEmpty(), TEXT("cannot define @TessellationOutputWinding attribute more than once"));
+		Strings.TessellationOutputWinding = TEXT("ccw");
+	}
+
+	void FHlslccHeaderWriter::WriteTessellationPartitioningInteger()
+	{
+		checkf(Strings.TessellationPartitioning.IsEmpty(), TEXT("cannot define @TessellationPartitioning attribute more than once"));
+		Strings.TessellationPartitioning = TEXT("integer");
+	}
+
+	void FHlslccHeaderWriter::WriteTessellationPartitioningFractionalOdd()
+	{
+		checkf(Strings.TessellationPartitioning.IsEmpty(), TEXT("cannot define @TessellationPartitioning attribute more than once"));
+		Strings.TessellationPartitioning = TEXT("fractional_odd");
+	}
+
+	void FHlslccHeaderWriter::WriteTessellationPartitioningFractionalEven()
+	{
+		checkf(Strings.TessellationPartitioning.IsEmpty(), TEXT("cannot define @TessellationPartitioning attribute more than once"));
+		Strings.TessellationPartitioning = TEXT("fractional_even");
+	}
+
+	void FHlslccHeaderWriter::WriteTessellationPatchesPerThreadGroup(uint32 NumPatches)
+	{
+		checkf(Strings.TessellationPatchesPerThreadGroup.IsEmpty(), TEXT("cannot define @TessellationPatchesPerThreadGroup attribute more than once"));
+		Strings.TessellationPatchesPerThreadGroup = FString::Printf(TEXT("%u"), NumPatches);
+	}
+
+	void FHlslccHeaderWriter::WriteTessellationPatchCountBuffer(uint32 BindingIndex)
+	{
+		checkf(Strings.TessellationPatchCountBuffer.IsEmpty(), TEXT("cannot define @TessellationPatchCountBuffer attribute more than once"));
+		Strings.TessellationPatchCountBuffer = FString::Printf(TEXT("%u"), BindingIndex);
+	}
+
+	void FHlslccHeaderWriter::WriteTessellationIndexBuffer(uint32 BindingIndex)
+	{
+		checkf(Strings.TessellationIndexBuffer.IsEmpty(), TEXT("cannot define @TessellationIndexBuffer attribute more than once"));
+		Strings.TessellationIndexBuffer = FString::Printf(TEXT("%u"), BindingIndex);
+	}
+
+	void FHlslccHeaderWriter::WriteTessellationHSOutBuffer(uint32 BindingIndex)
+	{
+		checkf(Strings.TessellationHSOutBuffer.IsEmpty(), TEXT("cannot define @TessellationHSOutBuffer attribute more than once"));
+		Strings.TessellationHSOutBuffer = FString::Printf(TEXT("%u"), BindingIndex);
+	}
+
+	void FHlslccHeaderWriter::WriteTessellationHSTFOutBuffer(uint32 BindingIndex)
+	{
+		checkf(Strings.TessellationHSTFOutBuffer.IsEmpty(), TEXT("cannot define @TessellationHSTFOutBuffer attribute more than once"));
+		Strings.TessellationHSTFOutBuffer = FString::Printf(TEXT("%u"), BindingIndex);
+	}
+
+	void FHlslccHeaderWriter::WriteTessellationControlPointOutBuffer(uint32 BindingIndex)
+	{
+		checkf(Strings.TessellationControlPointOutBuffer.IsEmpty(), TEXT("cannot define @TessellationControlPointOutBuffer attribute more than once"));
+		Strings.TessellationControlPointOutBuffer = FString::Printf(TEXT("%u"), BindingIndex);
+	}
+
+	void FHlslccHeaderWriter::WriteTessellationControlPointIndexBuffer(uint32 BindingIndex)
+	{
+		checkf(Strings.TessellationControlPointIndexBuffer.IsEmpty(), TEXT("cannot define @TessellationControlPointIndexBuffer attribute more than once"));
+		Strings.TessellationControlPointIndexBuffer = FString::Printf(TEXT("%u"), BindingIndex);
+	}
+
+	FString FHlslccHeaderWriter::ToString() const
+	{
+		FString MetaData;
+
+		auto PrintAttributes = [&MetaData](const TCHAR* Name, const FString& Value)
+		{
+			if (!Value.IsEmpty())
+			{
+				MetaData += FString::Printf(TEXT("// @%s: %s\n"), Name, *Value);
+			}
+		};
+
+		PrintAttributes(TEXT("Inputs"), Strings.InputAttributes);
+		PrintAttributes(TEXT("Outputs"), Strings.OutputAttributes);
+		PrintAttributes(TEXT("UniformBlocks"), Strings.UniformBlocks);
+		PrintAttributes(TEXT("PackedUB"), Strings.PackedUB + Strings.PackedUBFields);
+		PrintAttributes(TEXT("PackedGlobals"), Strings.PackedGlobals);
+		PrintAttributes(TEXT("PackedUBGlobalCopies"), Strings.PackedUBGlobalCopies);
+		PrintAttributes(TEXT("Samplers"), Strings.SRVs); // Was called "Samplers" in HLSLcc but serves as SRVs
+		PrintAttributes(TEXT("UAVs"), Strings.UAVs);
+		PrintAttributes(TEXT("SamplerStates"), Strings.SamplerStates);
+		PrintAttributes(TEXT("NumThreads"), Strings.NumThreads);
+		PrintAttributes(TEXT("SideTable"), Strings.SideTable);
+		PrintAttributes(TEXT("ArgumentBuffers"), Strings.ArgumentBuffers);
+		PrintAttributes(TEXT("TessellationOutputControlPoints"), Strings.TessellationOutputControlPoints);
+		PrintAttributes(TEXT("TessellationDomain"), Strings.TessellationDomain);
+		PrintAttributes(TEXT("TessellationInputControlPoints"), Strings.TessellationInputControlPoints);
+		PrintAttributes(TEXT("TessellationMaxTessFactor"), Strings.TessellationMaxTessFactor);
+		PrintAttributes(TEXT("TessellationOutputWinding"), Strings.TessellationOutputWinding);
+		PrintAttributes(TEXT("TessellationPartitioning"), Strings.TessellationPartitioning);
+		PrintAttributes(TEXT("TessellationPatchesPerThreadGroup"), Strings.TessellationPatchesPerThreadGroup);
+		PrintAttributes(TEXT("TessellationPatchCountBuffer"), Strings.TessellationPatchCountBuffer);
+		PrintAttributes(TEXT("TessellationIndexBuffer"), Strings.TessellationIndexBuffer);
+		PrintAttributes(TEXT("TessellationHSOutBuffer"), Strings.TessellationHSOutBuffer);
+		PrintAttributes(TEXT("TessellationControlPointOutBuffer"), Strings.TessellationControlPointOutBuffer);
+		PrintAttributes(TEXT("TessellationHSTFOutBuffer"), Strings.TessellationHSTFOutBuffer);
+		PrintAttributes(TEXT("TessellationControlPointIndexBuffer"), Strings.TessellationControlPointIndexBuffer);
+
+		return MetaData;
+	}
+
+} // namespace CrossCompiler
diff --git a/Engine/Source/Developer/ShaderCompilerCommon/Public/HlslccHeaderWriter.h b/Engine/Source/Developer/ShaderCompilerCommon/Public/HlslccHeaderWriter.h
new file mode 100644
--- /dev/null
+++ b/Engine/Source/Developer/ShaderCompilerCommon/Public/HlslccHeaderWriter.h
@@ -0,0 +1,101 @@
+// Copyright Epic Games, Inc. All Rights Reserved.
+
+#pragma once
+
+#include "CoreMinimal.h"
+
+
+// Forward declaration from <spirv_reflect.h>
+struct SpvReflectInterfaceVariable;
+
+// Cross compiler support/common functionality
+namespace CrossCompiler
+{
+	class SHADERCOMPILERCOMMON_API FHlslccHeaderWriter
+	{
+	public:
+
+		void WriteInputAttribute(const SpvReflectInterfaceVariable& Attribute);
+		void WriteInputAttribute(const TCHAR* AttributeName, const TCHAR* TypeSpecifier, int32 Location, bool bLocationPrefix, bool bLocationSuffix);
+		void WriteOutputAttribute(const SpvReflectInterfaceVariable& Attribute);
+		void WriteOutputAttribute(const TCHAR* AttributeName, const TCHAR* TypeSpecifier, int32 Location, bool bLocationPrefix, bool bLocationSuffix);
+		void WriteUniformBlock(const TCHAR* ResourceName, uint32 BindingIndex);
+		void WritePackedGlobal(const TCHAR* ResourceName, const TCHAR* TypeSpecifier, uint32 ByteOffset, uint32 ByteSize);
+		void WritePackedUB(uint32 BindingIndex);
+		void WritePackedUBField(const TCHAR* ResourceName, uint32 ByteOffset, uint32 ByteSize);
+		void WritePackedUBGlobalCopy(uint32 SourceCB, uint32 SourceOffset, uint32 DestCBIndex, uint32 DestCBPrecision, uint32 DestOffset, uint32 Size, bool bGroupFlattenedUBs = false);
+		void WriteSRV(const TCHAR* ResourceName, uint32 BindingIndex, uint32 Count = 1);
+		void WriteSRV(const TCHAR* ResourceName, uint32 BindingIndex, uint32 Count, const TArray<FString>& AssociatedResourceNames);
+		void WriteUAV(const TCHAR* ResourceName, uint32 BindingIndex, uint32 Count = 1);
+		void WriteSamplerState(const TCHAR* ResourceName, uint32 BindingIndex);
+		void WriteNumThreads(uint32 NumThreadsX, uint32 NumThreadsY, uint32 NumThreadsZ);
+
+		void WriteSideTable(const TCHAR* ResourceName, uint32 SideTableIndex);
+		void WriteArgumentBuffers(uint32 BindingIndex, const TArray<uint32>& ResourceIndices);
+
+		void WriteTessellationInputControlPoints(uint32 PatchSize);
+		void WriteTessellationOutputControlPoints(uint32 PatchSize);
+		void WriteTessellationMaxTessFactor(uint32 MaxTessFactor);
+
+		void WriteTessellationDomainTri();
+		void WriteTessellationDomainQuad();
+
+		void WriteTessellationOutputWindingCW();
+		void WriteTessellationOutputWindingCCW();
+
+		void WriteTessellationPartitioningInteger();
+		void WriteTessellationPartitioningFractionalOdd();
+		void WriteTessellationPartitioningFractionalEven();
+
+		void WriteTessellationPatchesPerThreadGroup(uint32 NumPatches);
+		void WriteTessellationPatchCountBuffer(uint32 BindingIndex);
+		void WriteTessellationIndexBuffer(uint32 BindingIndex);
+		void WriteTessellationHSOutBuffer(uint32 BindingIndex);
+		void WriteTessellationHSTFOutBuffer(uint32 BindingIndex);
+		void WriteTessellationControlPointOutBuffer(uint32 BindingIndex);
+		void WriteTessellationControlPointIndexBuffer(uint32 BindingIndex);
+
+		/** Returns the finalized meta data. */
+		FString ToString() const;
+
+	private:
+		void WriteIOAttribute(FString& OutMetaData, const TCHAR* AttributeName, const TCHAR* TypeSpecifier, int32 Location, bool bLocationPrefix, bool bLocationSuffix);
+		void WriteIOAttribute(FString& OutMetaData, const SpvReflectInterfaceVariable& Attribute, bool bIsInput);
+
+	private:
+		struct FMetaDataStrings
+		{
+			FString InputAttributes;
+			FString OutputAttributes;
+			FString UniformBlocks;
+			FString PackedGlobals;
+			FString PackedUB;
+			FString PackedUBFields;
+			FString PackedUBGlobalCopies;
+			FString SRVs; // Shader resource views (SRV) and samplers
+			FString UAVs; // Unordered access views (UAV)
+			FString SamplerStates;
+			FString NumThreads;
+			FString ExternalTextures; // External texture resources (Vulkan ES3.1 profile only)
+			FString SideTable; // Side table for additional indices, e.,g. "spvBufferSizeConstants(31)" (Metal only)
+			FString ArgumentBuffers; // Indirect argument buffers (Metal only)
+			FString TessellationOutputControlPoints;
+			FString TessellationDomain;
+			FString TessellationInputControlPoints;
+			FString TessellationMaxTessFactor;
+			FString TessellationOutputWinding;
+			FString TessellationPartitioning;
+			FString TessellationPatchesPerThreadGroup;
+			FString TessellationPatchCountBuffer;
+			FString TessellationIndexBuffer;
+			FString TessellationHSOutBuffer;
+			FString TessellationControlPointOutBuffer;
+			FString TessellationHSTFOutBuffer;
+			FString TessellationControlPointIndexBuffer;
+		};
+
+	private:
+		FMetaDataStrings Strings;
+		
+	};
+}
diff --git a/Engine/Source/Developer/ShaderCompilerCommon/ShaderCompilerCommon.Build.cs b/Engine/Source/Developer/ShaderCompilerCommon/ShaderCompilerCommon.Build.cs
--- a/Engine/Source/Developer/ShaderCompilerCommon/ShaderCompilerCommon.Build.cs
+++ b/Engine/Source/Developer/ShaderCompilerCommon/ShaderCompilerCommon.Build.cs
@@ -20,5 +20,6 @@ public class ShaderCompilerCommon : ModuleRules
 
 		// We only need a header containing definitions
 		PublicSystemIncludePaths.Add("ThirdParty/hlslcc/hlslcc/src/hlslcc_lib");
-    }
+		PublicSystemIncludePaths.Add("ThirdParty/SPIRV-Reflect/SPIRV-Reflect");
+	}
 }
diff --git a/Engine/Source/Developer/ShaderFormatOpenGL/Private/OpenGLShaderCompiler.cpp b/Engine/Source/Developer/ShaderFormatOpenGL/Private/OpenGLShaderCompiler.cpp
--- a/Engine/Source/Developer/ShaderFormatOpenGL/Private/OpenGLShaderCompiler.cpp
+++ b/Engine/Source/Developer/ShaderFormatOpenGL/Private/OpenGLShaderCompiler.cpp
@@ -7,6 +7,7 @@
 #include "Misc/Paths.h"
 #include "Serialization/MemoryWriter.h"
 #include "ShaderFormatOpenGL.h"
+#include "HlslccHeaderWriter.h"
 
 #if PLATFORM_WINDOWS
 #include "Windows/AllowWindowsPlatformTypes.h"
@@ -1376,13 +1377,8 @@ static bool CompileToGlslWithShaderConductor(
 		TArray<SpvReflectInterfaceVariable*> OutputVars;
 		TArray<SpvReflectBlockVariable*> ConstantBindings;
 		uint32 GlobalSetId = 32;
-		FString SRVString;
-		FString UAVString;
-		FString UBOString;
-		FString SMPString;
-		FString INPString;
-		FString OUTString;
-		FString PAKString;
+
+		CrossCompiler::FHlslccHeaderWriter CCHeaderWriter;
 		TArray<FString> Textures;
 		TArray<FString> Samplers;
 		uint32 UAVIndices = 0xffffffff;
@@ -1522,7 +1518,7 @@ static bool CompileToGlslWithShaderConductor(
 				TextureIndices &= ~(1llu << uint64(Index));
 				UAVIndices &= ~(1 << Index);
 
-				UAVString += FString::Printf(TEXT("%s%s(%u:%u)"), UAVString.Len() ? TEXT(",") : TEXT(""), UTF8_TO_TCHAR(Binding->name), Index, 1);
+				CCHeaderWriter.WriteUAV(UTF8_TO_TCHAR(Binding->name), Index);
 
 				SPVRResult = Reflection.ChangeDescriptorBindingNumbers(Binding, Index, GlobalSetId);
 				check(SPVRResult == SPV_REFLECT_RESULT_SUCCESS);
@@ -1538,7 +1534,7 @@ static bool CompileToGlslWithShaderConductor(
 				TextureIndices &= ~(1llu << uint64(Index));
 				UAVIndices &= ~(1 << Index);
 
-				UAVString += FString::Printf(TEXT("%s%s(%u:%u)"), UAVString.Len() ? TEXT(",") : TEXT(""), UTF8_TO_TCHAR(Binding->name), Index, 1);
+				CCHeaderWriter.WriteUAV(UTF8_TO_TCHAR(Binding->name), Index);
 
 				SPVRResult = Reflection.ChangeDescriptorBindingNumbers(Binding, Index, GlobalSetId);
 				check(SPVRResult == SPV_REFLECT_RESULT_SUCCESS);
@@ -1555,7 +1551,7 @@ static bool CompileToGlslWithShaderConductor(
 				TextureIndices &= ~(1llu << uint64(Index));
 				UAVIndices &= ~(1 << Index);
 
-				UAVString += FString::Printf(TEXT("%s%s(%u:%u)"), UAVString.Len() ? TEXT(",") : TEXT(""), UTF8_TO_TCHAR(Binding->name), Index, 1);
+				CCHeaderWriter.WriteUAV(UTF8_TO_TCHAR(Binding->name), Index);
 
 				SPVRResult = Reflection.ChangeDescriptorBindingNumbers(Binding, Index, GlobalSetId);
 				check(SPVRResult == SPV_REFLECT_RESULT_SUCCESS);
@@ -1599,10 +1595,7 @@ static bool CompileToGlslWithShaderConductor(
 					for (uint32 i = 0; i < Binding->block.member_count; i++)
 					{
 						SpvReflectBlockVariable& member = Binding->block.members[i];
-						uint32 MbrOffset = member.absolute_offset / sizeof(float);
-						uint32 MbrSize = member.size / sizeof(float);
-
-						PAKString += FString::Printf(TEXT("%s%s(h:%u,%u)"), PAKString.Len() ? TEXT(",") : TEXT(""), UTF8_TO_TCHAR(member.name), MbrOffset, MbrSize);
+						CCHeaderWriter.WritePackedGlobal(UTF8_TO_TCHAR(member.name), TEXT("h"), member.absolute_offset, member.size);
 					}
 				}
 				else
@@ -1614,7 +1607,7 @@ static bool CompileToGlslWithShaderConductor(
 					NewName += std::to_string(Index);
 					UniformVarNames[OldName] = NewName;
 					// Regular uniform buffer - we only care about the binding index
-					UBOString += FString::Printf(TEXT("%s%s(%u)"), UBOString.Len() ? TEXT(",") : TEXT(""), UTF8_TO_TCHAR(Binding->name), Index);
+					CCHeaderWriter.WriteUniformBlock(UTF8_TO_TCHAR(Binding->name), Index);
 				}
 
 				SPVRResult = Reflection.ChangeDescriptorBindingNumbers(Binding, Index, GlobalSetId);
@@ -1674,8 +1667,7 @@ static bool CompileToGlslWithShaderConductor(
 							}
 							auto const type = *member.type_description;
 
-							uint32 MbrOffset = Align(member.absolute_offset / sizeof(float), 4);
-							uint32 MbrSize = member.size / sizeof(float);
+							const uint32 MbrSize = member.size / sizeof(float);
 
 							FString TypeQualifier;
 
@@ -1692,7 +1684,7 @@ static bool CompileToGlslWithShaderConductor(
 
 							uint32& Offset = GlobalOffsets.FindOrAdd(TypeQualifier);
 
-							PAKString += FString::Printf(TEXT("%s%s(%s:%u,%u)"), PAKString.Len() ? TEXT(",") : TEXT(""), *MemberName, *TypeQualifier, Offset * 4, MbrSize);
+							CCHeaderWriter.WritePackedGlobal(*MemberName, *TypeQualifier, member.absolute_offset, member.size);
 
 							bool const bArray = type.traits.array.dims_count > 0;
 
@@ -1828,7 +1820,7 @@ static bool CompileToGlslWithShaderConductor(
 							FString Name = ANSI_TO_TCHAR(Var->name);
 							Name.ReplaceInline(TEXT("."), TEXT("_"));
 							OutputVarNames.Add(Name);
-							OUTString += FString::Printf(TEXT("%s%s;%d:out_Target%d"), OUTString.Len() ? TEXT(",") : TEXT(""), *TypeQualifier, Var->location, Var->location);
+							CCHeaderWriter.WriteOutputAttribute(TEXT("out_Target"), *TypeQualifier, Var->location, /*bLocationPrefix:*/ true, /*bLocationSuffix:*/ true);
 						}
 						else
 						{
@@ -1899,7 +1891,7 @@ static bool CompileToGlslWithShaderConductor(
 
 							FString Name = ANSI_TO_TCHAR(Var->name);
 							Name.ReplaceInline(TEXT("."), TEXT("_"));
-							OUTString += FString::Printf(TEXT("%s%s;%d:%s"), OUTString.Len() ? TEXT(",") : TEXT(""), *TypeQualifier, Location, *Name);
+							CCHeaderWriter.WriteOutputAttribute(*Name, *TypeQualifier, Location, /*bLocationPrefix:*/ true, /*bLocationSuffix:*/ false);
 						}
 					}
 				}
@@ -1991,7 +1983,7 @@ static bool CompileToGlslWithShaderConductor(
 						FString Name = ANSI_TO_TCHAR(Var->name);
 						Name.ReplaceInline(TEXT("."), TEXT("_"));
 						InputVarNames.Add(Name);
-						INPString += FString::Printf(TEXT("%s%s;%d:%s"), INPString.Len() ? TEXT(",") : TEXT(""), *TypeQualifier, Location, *Name);
+						CCHeaderWriter.WriteInputAttribute(*Name, *TypeQualifier, Location, /*bLocationPrefix:*/ true, /*bLocationSuffix:*/ false);
 					}
 				}
 			}
@@ -2213,7 +2205,7 @@ static bool CompileToGlslWithShaderConductor(
 				size_t SamplerPos = GlslSource.find("\nuniform ");
 
 				uint32 TextureIndex = 0;
-				for (FString& Texture : Textures)
+				for (const FString& Texture : Textures)
 				{
 					TArray<FString> UsedSamplers;
 					FString SamplerString;
@@ -2237,46 +2229,14 @@ static bool CompileToGlslWithShaderConductor(
 							SamplerString += FString::Printf(TEXT("%s%s"), SamplerString.Len() ? TEXT(",") : TEXT(""), *Sampler);
 						}
 					}
-					if (UsedSamplers.Num() > 0)
-					{
-						SRVString += FString::Printf(TEXT("%s%s(%u:%u[%s])"), SRVString.Len() ? TEXT(",") : TEXT(""), *Texture, TextureIndex, UsedSamplers.Num(), *SamplerString);
-						TextureIndex += UsedSamplers.Num();
-					}
-					else
-					{
-						SRVString += FString::Printf(TEXT("%s%s(%u:%u)"), SRVString.Len() ? TEXT(",") : TEXT(""), *Texture, TextureIndex++, 1);
-					}
+
+					const uint32 SamplerCount = FMath::Max(1, UsedSamplers.Num());
+					CCHeaderWriter.WriteSRV(*Texture, TextureIndex, SamplerCount, UsedSamplers);
+					TextureIndex += SamplerCount;
 				}
 
 				MetaData += TEXT("// Compiled by ShaderConductor\n");
-				if (INPString.Len())
-				{
-					MetaData += FString::Printf(TEXT("// @Inputs: %s\n"), *INPString);
-				}
-				if (OUTString.Len())
-				{
-					MetaData += FString::Printf(TEXT("// @Outputs: %s\n"), *OUTString);
-				}
-				if (UBOString.Len())
-				{
-					MetaData += FString::Printf(TEXT("// @UniformBlocks: %s\n"), *UBOString);
-				}
-				if (PAKString.Len())
-				{
-					MetaData += FString::Printf(TEXT("// @PackedGlobals: %s\n"), *PAKString);
-				}
-				if (SRVString.Len())
-				{
-					MetaData += FString::Printf(TEXT("// @Samplers: %s\n"), *SRVString);
-				}
-				if (UAVString.Len())
-				{
-					MetaData += FString::Printf(TEXT("// @UAVs: %s\n"), *UAVString);
-				}
-				if (SMPString.Len())
-				{
-					MetaData += FString::Printf(TEXT("// @SamplerStates: %s\n"), *SMPString);
-				}
+				MetaData += CCHeaderWriter.ToString();
 
 				// Merge meta data and GLSL source to output string
 				const int32 GlslShaderSourceLen = MetaData.Len() + static_cast<int32>(GlslSource.size()) + 1;
diff --git a/Engine/Source/Developer/VulkanShaderFormat/Private/VulkanShaderCompiler.cpp b/Engine/Source/Developer/VulkanShaderFormat/Private/VulkanShaderCompiler.cpp
--- a/Engine/Source/Developer/VulkanShaderFormat/Private/VulkanShaderCompiler.cpp
+++ b/Engine/Source/Developer/VulkanShaderFormat/Private/VulkanShaderCompiler.cpp
@@ -5,6 +5,7 @@
 #include "VulkanCommon.h"
 #include "ShaderPreprocessor.h"
 #include "ShaderCompilerCommon.h"
+#include "HlslccHeaderWriter.h"
 #include "hlslcc.h"
 
 #if PLATFORM_MAC || PLATFORM_WINDOWS || PLATFORM_LINUX
@@ -2013,17 +2014,6 @@ static void GatherSpirvReflectionBindingEntry(SpvReflectDescriptorBinding* InBin
 	} // switch
 }
 
-// Flattens the array dimensions of the interface variable (aka shader attribute), e.g. from float4[2][3] -> float4[6]
-static uint32 FlattenAttributeArrayDimension(const SpvReflectInterfaceVariable& Attribute, uint32 FirstArrayDim = 0)
-{
-	uint32 FlattenedArrayDim = 1;
-	for (uint32 ArrayDimIndex = FirstArrayDim; ArrayDimIndex < Attribute.array.dims_count; ++ArrayDimIndex)
-	{
-		FlattenedArrayDim *= Attribute.array.dims[ArrayDimIndex];
-	}
-	return FlattenedArrayDim;
-}
-
 static void GatherSpirvReflectionBindings(
 	spv_reflect::ShaderModule&	Reflection,
 	FSpirvReflectionBindings&	OutBindings,
@@ -2107,282 +2097,10 @@ static void GatherSpirvReflectionBindings(
 		for (auto const& BindingEntry : Bindings)
 		{
 			GatherSpirvReflectionBindingEntry(BindingEntry, OutBindings);
-		} // for
-	}
-}
-
-static void ConvertMetaDataTypeSpecifierPrimary(const SpvReflectTypeDescription& TypeSpecifier, FString& OutTypeName, uint32& OutTypeBitWidth, bool bBaseTypeOnly)
-{
-	// Generate prefix for base type
-	if (TypeSpecifier.type_flags & SPV_REFLECT_TYPE_FLAG_BOOL)
-	{
-		OutTypeName += TEXT('b');
-		OutTypeBitWidth = 8;
-	}
-	else if (TypeSpecifier.type_flags & SPV_REFLECT_TYPE_FLAG_INT)
-	{
-		if (TypeSpecifier.traits.numeric.scalar.signedness)
-		{
-			OutTypeName += TEXT('i');
-		}
-		else
-		{
-			OutTypeName += TEXT('u');
-		}
-		OutTypeBitWidth = 32;
-	}
-	else if (TypeSpecifier.type_flags & SPV_REFLECT_TYPE_FLAG_FLOAT)
-	{
-		if (TypeSpecifier.traits.numeric.scalar.width == 16)
-		{
-			OutTypeName += TEXT('h');
-			OutTypeBitWidth = 16;
-		}
-		else
-		{
-			OutTypeName += TEXT('f');
-			OutTypeBitWidth = 32;
-		}
-	}
-
-	if (!bBaseTypeOnly)
-	{
-		// Generate number for vector size
-		const SpvReflectTypeFlags SpvScalarTypeFlags = (SPV_REFLECT_TYPE_FLAG_BOOL | SPV_REFLECT_TYPE_FLAG_INT | SPV_REFLECT_TYPE_FLAG_FLOAT);
-		if (TypeSpecifier.type_flags & SPV_REFLECT_TYPE_FLAG_VECTOR)
-		{
-			static const TCHAR* VectorDims = TEXT("1234");
-			const uint32 VectorSize = TypeSpecifier.traits.numeric.vector.component_count;
-			check(VectorSize >= 1 && VectorSize <= 4);
-			OutTypeName += VectorDims[VectorSize - 1];
-		}
-		else if (TypeSpecifier.type_flags & SPV_REFLECT_TYPE_FLAG_MATRIX)
-		{
-			//TODO
-		}
-		else if ((TypeSpecifier.type_flags & SpvScalarTypeFlags) != 0)
-		{
-			OutTypeName += TEXT('1'); // add single scalar component
 		}
 	}
 }
 
-static FString ConvertMetaDataTypeSpecifier(const SpvReflectTypeDescription& TypeSpecifier, uint32* OutTypeBitWidth = nullptr, bool bBaseTypeOnly = false)
-{
-	FString TypeName;
-	uint32 TypeBitWidth = sizeof(float) * 8;
-	ConvertMetaDataTypeSpecifierPrimary(TypeSpecifier, TypeName, TypeBitWidth, bBaseTypeOnly);
-	if (OutTypeBitWidth)
-	{
-		*OutTypeBitWidth = TypeBitWidth;
-	}
-	return TypeName;
-}
-
-static const TCHAR* SpvBuiltinToString(const SpvBuiltIn BuiltIn)
-{
-	switch (BuiltIn)
-	{
-	case SpvBuiltInPosition:					return TEXT("gl_Position");
-	case SpvBuiltInPointSize:					return TEXT("gl_PointSize");
-	case SpvBuiltInClipDistance:				return TEXT("gl_ClipDistance");
-	case SpvBuiltInCullDistance:				return TEXT("gl_CullDistance");
-	case SpvBuiltInVertexId:					return TEXT("gl_VertexID");
-	case SpvBuiltInInstanceId:					return TEXT("gl_InstanceID");
-	case SpvBuiltInPrimitiveId:					return TEXT("gl_PrimitiveID");
-	case SpvBuiltInInvocationId:				return TEXT("gl_InvocationID");
-	case SpvBuiltInLayer:						return TEXT("gl_Layer");
-	case SpvBuiltInViewportIndex:				return TEXT("gl_ViewportIndex");
-	case SpvBuiltInTessLevelOuter:				return TEXT("gl_TessLevelOuter");
-	case SpvBuiltInTessLevelInner:				return TEXT("gl_TessLevelInner");
-	case SpvBuiltInTessCoord:					return TEXT("gl_TessCoord");
-	case SpvBuiltInPatchVertices:				return TEXT("gl_PatchVertices");
-	case SpvBuiltInFragCoord:					return TEXT("gl_FragCoord");
-	case SpvBuiltInPointCoord:					return TEXT("gl_PointCoord");
-	case SpvBuiltInFrontFacing:					return TEXT("gl_FrontFacing");
-	case SpvBuiltInSampleId:					return TEXT("gl_SampleID");
-	case SpvBuiltInSamplePosition:				return TEXT("gl_SamplePosition");
-	case SpvBuiltInSampleMask:					return TEXT("gl_SampleMask");
-	case SpvBuiltInFragDepth:					return TEXT("gl_FragDepth");
-	case SpvBuiltInHelperInvocation:			return TEXT("gl_HelperInvocation");
-	case SpvBuiltInNumWorkgroups:				return TEXT("gl_NumWorkgroups");
-	case SpvBuiltInWorkgroupSize:				return TEXT("gl_WorkgroupSize");
-	case SpvBuiltInWorkgroupId:					return TEXT("gl_WorkgroupID");
-	case SpvBuiltInLocalInvocationId:			return TEXT("gl_LocalInvocationID");
-	case SpvBuiltInGlobalInvocationId:			return TEXT("gl_GlobalInvocationID");
-	case SpvBuiltInLocalInvocationIndex:		return TEXT("gl_LocalInvocationIndex");
-	case SpvBuiltInWorkDim:						return TEXT("gl_WorkDim");
-	case SpvBuiltInGlobalSize:					return TEXT("gl_GlobalSize");
-	case SpvBuiltInEnqueuedWorkgroupSize:		return TEXT("gl_EnqueuedWorkgroupSize");
-	case SpvBuiltInGlobalOffset:				return TEXT("gl_GlobalOffset");
-	case SpvBuiltInGlobalLinearId:				return TEXT("gl_GlobalLinearID");
-	case SpvBuiltInSubgroupSize:				return TEXT("gl_SubgroupSize");
-	case SpvBuiltInSubgroupMaxSize:				return TEXT("gl_SubgroupMaxSize");
-	case SpvBuiltInNumSubgroups:				return TEXT("gl_NumSubgroups");
-	case SpvBuiltInNumEnqueuedSubgroups:		return TEXT("gl_NumEnqueuedSubgroups");
-	case SpvBuiltInSubgroupId:					return TEXT("gl_SubgroupID");
-	case SpvBuiltInSubgroupLocalInvocationId:	return TEXT("gl_SubgroupLocalInvocationID");
-	case SpvBuiltInVertexIndex:					return TEXT("gl_VertexIndex");
-	case SpvBuiltInInstanceIndex:				return TEXT("gl_InstanceIndex");
-	case SpvBuiltInSubgroupEqMask:				return TEXT("gl_SubgroupEqMask");
-	case SpvBuiltInSubgroupGeMask:				return TEXT("gl_SubgroupGeMask");
-	case SpvBuiltInSubgroupGtMask:				return TEXT("gl_SubgroupGtMask");
-	case SpvBuiltInSubgroupLeMask:				return TEXT("gl_SubgroupLeMask");
-	case SpvBuiltInSubgroupLtMask:				return TEXT("gl_SubgroupLtMask");
-	case SpvBuiltInBaseVertex:					return TEXT("gl_BaseVertex");
-	case SpvBuiltInBaseInstance:				return TEXT("gl_BaseInstance");
-	case SpvBuiltInDrawIndex:					return TEXT("gl_DrawIndex");
-	case SpvBuiltInDeviceIndex:					return TEXT("gl_DeviceIndex");
-	case SpvBuiltInViewIndex:					return TEXT("gl_ViewIndex");
-	}
-	return nullptr;
-}
-
-static FString ConvertAttributeToMetaDataSemantic(const ANSICHAR* AttributeName, const SpvBuiltIn BuiltIn, bool bIsInput)
-{
-	if (const TCHAR* BuiltInName = SpvBuiltinToString(BuiltIn))
-	{
-		return FString(BuiltInName);
-	}
-	else
-	{
-		check(AttributeName != nullptr && *AttributeName != '\0');
-		FString InSemantic = ANSI_TO_TCHAR(AttributeName);
-		FString OutSemantic = (bIsInput ? TEXT("in_") : TEXT("out_"));
-
-		if (InSemantic.StartsWith(TEXT("SV_")))
-		{
-			OutSemantic += InSemantic.Right(InSemantic.Len() - 3);
-		}
-		else
-		{
-			OutSemantic += InSemantic;
-		}
-
-		return OutSemantic;
-	}
-}
-
-// Returns the string position where the index in the specified HLSL semantic beings, e.g. "SV_Target2" -> 9, "SV_Target" -> INDEX_NONE
-static int32 FindIndexInHlslSemantic(const FString& Semantic)
-{
-	int32 Index = Semantic.Len();
-	if (Index > 0 && FChar::IsDigit(Semantic[Index - 1]))
-	{
-		while (Index > 0 && FChar::IsDigit(Semantic[Index - 1]))
-		{
-			--Index;
-		}
-		return Index;
-	}
-	return INDEX_NONE;
-}
-
-static void BuildShaderInterfaceVariableMetaData(const SpvReflectInterfaceVariable& Attribute, FString& OutMetaData, bool bIsInput)
-{
-	// Ignore interface variables that are only generated for intermediate results
-	if (CrossCompiler::FShaderConductorContext::IsIntermediateSpirvOutputVariable(Attribute.name))
-	{
-		return;
-	}
-
-	const FString TypeSpecifier = ConvertMetaDataTypeSpecifier(*Attribute.type_description);
-	FString Semantic = ConvertAttributeToMetaDataSemantic(Attribute.semantic, Attribute.built_in, bIsInput);
-
-	if (Attribute.array.dims_count > 0)
-	{
-		// Get semantic without index, e.g. "out_Target0" -> "out_Target"
-		const int32 SemanticIndexPos = FindIndexInHlslSemantic(Semantic);
-		if (SemanticIndexPos != INDEX_NONE)
-		{
-			Semantic = Semantic.Left(SemanticIndexPos);
-		}
-
-		if (Attribute.location == -1)
-		{
-			// Flatten array dimensions, e.g. from float4[3][2] -> float4[6]
-			const uint32 FlattenedArrayDim = FlattenAttributeArrayDimension(Attribute);
-
-			// Emit one output slot for each array element, e.g. "out float4 OutColor[2] : SV_Target0" occupies output slot SV_Target0 and SV_Target1.
-			for (uint32 FlattenedArrayIndex = 0; FlattenedArrayIndex < FlattenedArrayDim; ++FlattenedArrayIndex)
-			{
-				// If there is no binding slot, emit output as system value array such as "gl_SampleMask[]"
-				OutMetaData += FString::Printf(
-					TEXT("%s%s;%d:%s[%d]"),
-					!OutMetaData.IsEmpty() ? TEXT(",") : TEXT(""),
-					*TypeSpecifier, // type specifier
-					Attribute.location,
-					*Semantic,
-					FlattenedArrayIndex
-				);
-			}
-		}
-		else if (!bIsInput)
-		{
-			//NOTE: For some reason, the meta data for output slot arrays must be entirely flattened, including the outer most array dimension
-			// Flatten array dimensions, e.g. from float4[3][2] -> float4[6]
-			const uint32 FlattenedArrayDim = FlattenAttributeArrayDimension(Attribute);
-
-			// Emit one output slot for each array element, e.g. "out float4 OutColor[2] : SV_Target0" occupies output slot SV_Target0 and SV_Target1.
-			for (uint32 FlattenedArrayIndex = 0; FlattenedArrayIndex < FlattenedArrayDim; ++FlattenedArrayIndex)
-			{
-				const uint32 BindingSlot = Attribute.location + FlattenedArrayIndex;
-				OutMetaData += FString::Printf(
-					TEXT("%s%s;%d:%s%d"),
-					!OutMetaData.IsEmpty() ? TEXT(",") : TEXT(""),
-					*TypeSpecifier, // Type specifier
-					BindingSlot,
-					*Semantic,
-					BindingSlot
-				);
-			}
-		}
-		else if (Attribute.array.dims_count >= 2)
-		{
-			// Flatten array dimensions, e.g. from float4[3][2] -> float4[6]
-			const uint32 FlattenedArrayDim = FlattenAttributeArrayDimension(Attribute, 1);
-
-			// Emit one output slot for each array element, e.g. "out float4 OutColor[2] : SV_Target0" occupies output slot SV_Target0 and SV_Target1.
-			for (uint32 FlattenedArrayIndex = 0; FlattenedArrayIndex < FlattenedArrayDim; ++FlattenedArrayIndex)
-			{
-				const uint32 BindingSlot = Attribute.location + FlattenedArrayIndex;
-				OutMetaData += FString::Printf(
-					TEXT("%s%s[%d];%d:%s%d"),
-					!OutMetaData.IsEmpty() ? TEXT(",") : TEXT(""),
-					*TypeSpecifier, // Type specifier
-					Attribute.array.dims[0], // Outer most array dimension
-					BindingSlot,
-					*Semantic,
-					BindingSlot
-				);
-			}
-		}
-		else
-		{
-			const uint32 BindingSlot = Attribute.location;
-			OutMetaData += FString::Printf(
-				TEXT("%s%s[%d];%d:%s%d"),
-				!OutMetaData.IsEmpty() ? TEXT(",") : TEXT(""),
-				*TypeSpecifier, // Type specifier
-				Attribute.array.dims[0], // Outer most array dimension
-				BindingSlot,
-				*Semantic,
-				BindingSlot
-			);
-		}
-	}
-	else
-	{
-		OutMetaData += FString::Printf(
-			TEXT("%s%s;%d:%s"),
-			!OutMetaData.IsEmpty() ? TEXT(",") : TEXT(""),
-			*TypeSpecifier, // type specifier
-			Attribute.location,
-			*Semantic
-		);
-	}
-}
-
 static void BuildShaderOutputFromSpirv(
 	FSpirv&						Spirv,
 	const FShaderCompilerInput& Input,
@@ -2393,7 +2111,6 @@ static void BuildShaderOutputFromSpirv(
 {
 	FShaderParameterMap& ParameterMap = Output.ParameterMap;
 
-	FString UAVString, SRVString, SMPString, UBOString, GLOString, PAKString, INPString, OUTString, WKGString;
 	SpvReflectResult SpvResult = SPV_REFLECT_RESULT_SUCCESS;
 
 	uint8 UAVIndices = 0xff;
@@ -2517,15 +2234,17 @@ static void BuildShaderOutputFromSpirv(
 	// Sort binding table
 	BindingTable.SortBindings();
 
+	CrossCompiler::FHlslccHeaderWriter CCHeaderWriter;
+
 	// Iterate over all resource bindings grouped by resource type
 	for (const SpvReflectInterfaceVariable* Attribute : Bindings.InputAttributes)
 	{
-		BuildShaderInterfaceVariableMetaData(*Attribute, INPString, /*bIsInput:*/ true);
+		CCHeaderWriter.WriteInputAttribute(*Attribute);
 	}
 
 	for (const SpvReflectInterfaceVariable* Attribute : Bindings.OutputAttributes)
 	{
-		BuildShaderInterfaceVariableMetaData(*Attribute, OUTString, /*bIsInput:*/ false);
+		CCHeaderWriter.WriteOutputAttribute(*Attribute);
 	}
 
 	int32 UBOBindings = 0, UAVBindings = 0, SRVBindings = 0, SMPBindings = 0, GLOBindings = 0;
@@ -2550,7 +2269,7 @@ static void BuildShaderOutputFromSpirv(
 
 			Spirv.ReflectionInfo.Add(FSpirv::FEntry(UBOGlobalsNameSpv, BindingIndex));
 
-			UBOString += FString::Printf(TEXT("%s%s(%u)"), UBOString.Len() ? TEXT(",") : TEXT(""), TEXT("_Globals_h"), UBOBindings++);
+			CCHeaderWriter.WriteUniformBlock(TEXT("_Globals_h"), UBOBindings++);
 
 			// Register all uniform buffer members as loose data
 			FString MbrString;
@@ -2558,25 +2277,11 @@ static void BuildShaderOutputFromSpirv(
 			for (uint32 MemberIndex = 0; MemberIndex < Binding->block.member_count; ++MemberIndex)
 			{
 				const SpvReflectBlockVariable* Member = &(Binding->block.members[MemberIndex]);
-
 				const FString MemberName(ANSI_TO_TCHAR(Member->name));
-				uint32 MemberTypeBitWidth = sizeof(float) * 8;
-				const FString TypeSpecifier = ConvertMetaDataTypeSpecifier(*Member->type_description, &MemberTypeBitWidth, /*bBaseTypeOnly:*/ true);
-				const uint32 MemberOffset = Member->absolute_offset / sizeof(float);
-				const uint32 MemberComponentCount = Member->size * 8 / MemberTypeBitWidth;
-
-				MbrString += FString::Printf(
-					TEXT("%s%s(%s:%u,%u)"),
-					MbrString.Len() ? TEXT(",") : TEXT(""),
-					*MemberName,
-					TEXT("h"),//*TypeSpecifier,
-					MemberOffset,
-					MemberComponentCount
-				);
+				CCHeaderWriter.WritePackedGlobal(*MemberName, TEXT("h"), Member->absolute_offset, Member->size);
 			}
 
-			GLOString += MbrString;
-
+			// Stop after we found $Globals uniform buffer
 			break;
 		}
 	}
@@ -2596,7 +2301,7 @@ static void BuildShaderOutputFromSpirv(
 
 			Spirv.ReflectionInfo.Add(FSpirv::FEntry(ResourceName, BindingIndex));
 
-			UBOString += FString::Printf(TEXT("%s%s(%u)"), UBOString.Len() ? TEXT(",") : TEXT(""), *ResourceName, UBOBindings++);
+			CCHeaderWriter.WriteUniformBlock(*ResourceName, UBOBindings++);
 		}
 	}
 
@@ -2608,7 +2313,6 @@ static void BuildShaderOutputFromSpirv(
 		check(SpvResult == SPV_REFLECT_RESULT_SUCCESS);
 
 		const FString ResourceName(ANSI_TO_TCHAR(Binding->name));
-		//IAString += FString::Printf(TEXT("%s%s(%u:%u)"), IAString.Len() ? TEXT(",") : TEXT(""), *ResourceName, IABindings++, 1);
 
 		Spirv.ReflectionInfo.Add(FSpirv::FEntry(ResourceName, BindingIndex));
 	}
@@ -2621,7 +2325,7 @@ static void BuildShaderOutputFromSpirv(
 		check(SpvResult == SPV_REFLECT_RESULT_SUCCESS);
 
 		const FString ResourceName(ANSI_TO_TCHAR(Binding->name));
-		UAVString += FString::Printf(TEXT("%s%s(%u:%u)"), UAVString.Len() ? TEXT(",") : TEXT(""), *ResourceName, UAVBindings++, 1);
+		CCHeaderWriter.WriteUAV(*ResourceName, UAVBindings++);
 
 		Spirv.ReflectionInfo.Add(FSpirv::FEntry(ResourceName, BindingIndex));
 	}
@@ -2634,7 +2338,7 @@ static void BuildShaderOutputFromSpirv(
 		check(SpvResult == SPV_REFLECT_RESULT_SUCCESS);
 
 		const FString ResourceName(ANSI_TO_TCHAR(Binding->name));
-		UAVString += FString::Printf(TEXT("%s%s(%u:%u)"), UAVString.Len() ? TEXT(",") : TEXT(""), *ResourceName, UAVBindings++, 1);
+		CCHeaderWriter.WriteUAV(*ResourceName, UAVBindings++);
 
 		Spirv.ReflectionInfo.Add(FSpirv::FEntry(ResourceName, BindingIndex));
 	}
@@ -2647,7 +2351,7 @@ static void BuildShaderOutputFromSpirv(
 		check(SpvResult == SPV_REFLECT_RESULT_SUCCESS);
 
 		const FString ResourceName(ANSI_TO_TCHAR(Binding->name));
-		UAVString += FString::Printf(TEXT("%s%s(%u:%u)"), UAVString.Len() ? TEXT(",") : TEXT(""), *ResourceName, UAVBindings++, 1);
+		CCHeaderWriter.WriteUAV(*ResourceName, UAVBindings++);
 
 		Spirv.ReflectionInfo.Add(FSpirv::FEntry(ResourceName, BindingIndex));
 	}
@@ -2660,7 +2364,7 @@ static void BuildShaderOutputFromSpirv(
 		check(SpvResult == SPV_REFLECT_RESULT_SUCCESS);
 
 		const FString ResourceName(ANSI_TO_TCHAR(Binding->name));
-		SRVString += FString::Printf(TEXT("%s%s(%u:%u)"), SRVString.Len() ? TEXT(",") : TEXT(""), *ResourceName, SRVBindings++, 1);
+		CCHeaderWriter.WriteSRV(*ResourceName, UAVBindings++);
 
 		Spirv.ReflectionInfo.Add(FSpirv::FEntry(ResourceName, BindingIndex));
 	}
@@ -2673,7 +2377,7 @@ static void BuildShaderOutputFromSpirv(
 		check(SpvResult == SPV_REFLECT_RESULT_SUCCESS);
 
 		const FString ResourceName(ANSI_TO_TCHAR(Binding->name));
-		SRVString += FString::Printf(TEXT("%s%s(%u:%u)"), SRVString.Len() ? TEXT(",") : TEXT(""), *ResourceName, SRVBindings++, 1);
+		CCHeaderWriter.WriteSRV(*ResourceName, SRVBindings++);
 
 		Spirv.ReflectionInfo.Add(FSpirv::FEntry(ResourceName, BindingIndex));
 	}
@@ -2688,24 +2392,20 @@ static void BuildShaderOutputFromSpirv(
 		const FString ResourceName(ANSI_TO_TCHAR(Binding->name));
 		if (Binding->usage_binding_count > 0)
 		{
-			SRVString += FString::Printf(TEXT("%s%s(%u:%u["), SRVString.Len() ? TEXT(",") : TEXT(""), *ResourceName, SRVBindings++, 1);
+			TArray<FString> AssociatedResourceNames;
+			AssociatedResourceNames.SetNum(Binding->usage_binding_count);
 
 			for (uint32 UsageIndex = 0; UsageIndex < Binding->usage_binding_count; ++UsageIndex)
 			{
 				const SpvReflectDescriptorBinding* AssociatedResource = Binding->usage_bindings[UsageIndex];
-				const FString AssociatedResourceName(ANSI_TO_TCHAR(AssociatedResource->name));
-				if (UsageIndex > 0)
-				{
-					SRVString += TEXT(",");
-				}
-				SRVString += AssociatedResourceName;
+				AssociatedResourceNames[UsageIndex] = ANSI_TO_TCHAR(AssociatedResource->name);
 			}
 
-			SRVString += TEXT("])");
+			CCHeaderWriter.WriteSRV(*ResourceName, SRVBindings++, /*Count:*/ 1, AssociatedResourceNames);
 		}
 		else
 		{
-			SRVString += FString::Printf(TEXT("%s%s(%u:%u)"), SRVString.Len() ? TEXT(",") : TEXT(""), *ResourceName, SRVBindings++, 1);
+			CCHeaderWriter.WriteSRV(*ResourceName, SRVBindings++);
 		}
 
 		Spirv.ReflectionInfo.Add(FSpirv::FEntry(ResourceName, BindingIndex));
@@ -2726,7 +2426,7 @@ static void BuildShaderOutputFromSpirv(
 		{
 //			if (*UseCount >= 2)
 			{
-				SMPString += FString::Printf(TEXT("%s%u:%s"), SMPString.Len() ? TEXT(",") : TEXT(""), SMPBindings++, *ResourceName);
+				CCHeaderWriter.WriteSamplerState(*ResourceName, SMPBindings++);
 			}
 		}
 	}
@@ -2736,42 +2436,7 @@ static void BuildShaderOutputFromSpirv(
 
 	FString MetaData;// = FString::Printf(TEXT("// ! %s/%s:%s\n"), *Input.DebugGroupName, *Input.GetSourceFilename(), *Input.EntryPointName);
 	MetaData += TEXT("// Compiled by ShaderConductor\n");
-	if (INPString.Len())
-	{
-		MetaData += FString::Printf(TEXT("// @Inputs: %s\n"), *INPString);
-	}
-	if (OUTString.Len())
-	{
-		MetaData += FString::Printf(TEXT("// @Outputs: %s\n"), *OUTString);
-	}
-	if (UBOString.Len())
-	{
-		MetaData += FString::Printf(TEXT("// @UniformBlocks: %s\n"), *UBOString);
-	}
-	if (GLOString.Len())
-	{
-		MetaData += FString::Printf(TEXT("// @PackedGlobals: %s\n"), *GLOString);
-	}
-	if (PAKString.Len())
-	{
-		MetaData += FString::Printf(TEXT("// @PackedUBGlobalCopies: %s\n"), *PAKString);
-	}
-	if (SRVString.Len())
-	{
-		MetaData += FString::Printf(TEXT("// @Samplers: %s\n"), *SRVString);
-	}
-	if (UAVString.Len())
-	{
-		MetaData += FString::Printf(TEXT("// @UAVs: %s\n"), *UAVString);
-	}
-	if (SMPString.Len())
-	{
-		MetaData += FString::Printf(TEXT("// @SamplerStates: %s\n"), *SMPString);
-	}
-	if (WKGString.Len())
-	{
-		MetaData += FString::Printf(TEXT("// @NumThreads: %s\n"), *WKGString);
-	}
+	MetaData += CCHeaderWriter.ToString();
 
 	Output.Target = Input.Target;
 
