Add missing #include directive in DXCWrapper.cpp.

#rnx

diff --git a/Engine/Source/Developer/ShaderCompilerCommon/Private/DXCWrapper.cpp b/Engine/Source/Developer/ShaderCompilerCommon/Private/DXCWrapper.cpp
--- a/Engine/Source/Developer/ShaderCompilerCommon/Private/DXCWrapper.cpp
+++ b/Engine/Source/Developer/ShaderCompilerCommon/Private/DXCWrapper.cpp
@@ -2,6 +2,7 @@
 
 #include "DXCWrapper.h"
 #include "HAL/PlatformProcess.h"
+#include "HAL/FileManager.h"
 #include "Templates/RefCounting.h"
 #include "Misc/Paths.h"
 #include "Hash/CityHash.h"
