Always read DLL file to hash DXC version.

#rnx

diff --git a/Engine/Source/Developer/ShaderCompilerCommon/Private/DXCWrapper.cpp b/Engine/Source/Developer/ShaderCompilerCommon/Private/DXCWrapper.cpp
--- a/Engine/Source/Developer/ShaderCompilerCommon/Private/DXCWrapper.cpp
+++ b/Engine/Source/Developer/ShaderCompilerCommon/Private/DXCWrapper.cpp
@@ -32,19 +32,17 @@ static uint64 GetLoadedModuleVersion(const TCHAR* ModuleName)
 	}
 	check(PathLen < UE_ARRAY_COUNT(DllPath));
 
-	// Official shader compiler binaries will have the the version, but anything that we compile ourselves won't
-	uint64 FileVersion = FPlatformMisc::GetFileVersion(DllPath);
-	if (FileVersion == 0)
+	uint64 FileVersion = 0;
+
+	TUniquePtr<FArchive> FileReader(IFileManager::Get().CreateFileReader(DllPath));
+	if (FileReader)
 	{
-		TUniquePtr<FArchive> FileReader(IFileManager::Get().CreateFileReader(DllPath));
-		if (FileReader)
-		{
-			TArray<char> FileContents;
-			FileContents.SetNumUninitialized(FileReader->TotalSize());
-			FileReader->Serialize(FileContents.GetData(), FileContents.Num());
-			FileVersion = CityHash64(FileContents.GetData(), FileContents.Num());
-		}
+		TArray<char> FileContents;
+		FileContents.SetNumUninitialized(FileReader->TotalSize());
+		FileReader->Serialize(FileContents.GetData(), FileContents.Num());
+		FileVersion = CityHash64(FileContents.GetData(), FileContents.Num());
 	}
+
 	return FileVersion;
 #else // PLATFORM_WINDOWS
 	return 0;
