UE5 - Minor refactor in prep for other changes

diff --git a/Engine/Source/Developer/ShaderCompilerCommon/Private/ShaderConductorContext.cpp b/Engine/Source/Developer/ShaderCompilerCommon/Private/ShaderConductorContext.cpp
--- a/Engine/Source/Developer/ShaderCompilerCommon/Private/ShaderConductorContext.cpp
+++ b/Engine/Source/Developer/ShaderCompilerCommon/Private/ShaderConductorContext.cpp
@@ -273,6 +273,27 @@ namespace CrossCompiler
 		checkf(OutTargetDesc.version, TEXT("Unsupported target shader version for Metal family: %d"), InTarget.Version);
 	}
 
+	// Converts a map of string pairs to a C-Style macro defines array
+	static void ConvertStringMapToMacroDefines(const TMap<FString,FString>& InPairs, TArray<TPair<TArray<ANSICHAR>, TArray<ANSICHAR>>>& OutPairs, TArray<ShaderConductor::MacroDefine>& OutPairRefs)
+	{
+		// Convert map into an array container
+		TArray<ANSICHAR> Name, Value;
+		for (const TPair<FString, FString>& Iter : InPairs)
+		{
+			ConvertFStringToAnsiString(Iter.Key, Name);
+			ConvertFStringToAnsiString(Iter.Value, Value);
+			OutPairs.Emplace(MoveTemp(Name), MoveTemp(Value));
+		}
+
+		// Store references after all elements have been added to the container so the pointers remain valid
+		OutPairRefs.SetNum(OutPairs.Num());
+		for (int32 Index = 0; Index < OutPairs.Num(); ++Index)
+		{
+			OutPairRefs[Index].name = OutPairs[Index].Key.GetData();
+			OutPairRefs[Index].value = OutPairs[Index].Value.GetData();
+		}
+	}
+
 	static void ConvertScTargetDesc(FShaderConductorContext::FShaderConductorIntermediates& Intermediates, const FShaderConductorTarget& InTarget, ShaderConductor::Compiler::TargetDesc& OutTargetDesc)
 	{
 		// Convert FString to ANSI string and store them as intermediates
@@ -291,21 +312,7 @@ namespace CrossCompiler
 		}
 
 		// Convert flags map into an array container
-		TArray<ANSICHAR> FlagName, FlagValue;
-		for (const TPair<FString, FString>& Iter : InTarget.CompileFlags.GetDefinitionMap())
-		{
-			ConvertFStringToAnsiString(Iter.Key, FlagName);
-			ConvertFStringToAnsiString(Iter.Value, FlagValue);
-			Intermediates.Flags.Emplace(MoveTemp(FlagName), MoveTemp(FlagValue));
-		}
-
-		// Store references after all elements have been added to the container so the pointers remain valid
-		Intermediates.FlagRefs.SetNum(Intermediates.Flags.Num());
-		for (int32 Index = 0; Index < Intermediates.Flags.Num(); ++Index)
-		{
-			Intermediates.FlagRefs[Index].name = Intermediates.Flags[Index].Key.GetData();
-			Intermediates.FlagRefs[Index].value = Intermediates.Flags[Index].Value.GetData();
-		}
+		ConvertStringMapToMacroDefines(InTarget.CompileFlags.GetDefinitionMap(), Intermediates.Flags, Intermediates.FlagRefs);
 
 		OutTargetDesc.options = Intermediates.FlagRefs.GetData();
 		OutTargetDesc.numOptions = static_cast<uint32>(Intermediates.FlagRefs.Num());
@@ -349,27 +356,6 @@ namespace CrossCompiler
 		}
 	}
 
-	static void ConvertScDefines(FShaderConductorContext::FShaderConductorIntermediates& Intermediates, const FShaderCompilerDefinitions& InDefinitions)
-	{
-		// Convert FString to ANSI string for each macro definition and its value
-		TArray<ANSICHAR> DefineName, DefineValue;
-
-		for (const TPair<FString, FString>& Iter : InDefinitions.GetDefinitionMap())
-		{
-			ConvertFStringToAnsiString(Iter.Key, DefineName);
-			ConvertFStringToAnsiString(Iter.Value, DefineValue);
-			Intermediates.Defines.Emplace(MoveTemp(DefineName), MoveTemp(DefineValue));
-		}
-
-		// Store references after all elements have been added to the container so the pointers remain valid
-		Intermediates.DefineRefs.SetNum(Intermediates.Defines.Num());
-		for (int32 Index = 0; Index < Intermediates.Defines.Num(); ++Index)
-		{
-			Intermediates.DefineRefs[Index].name = Intermediates.Defines[Index].Key.GetData();
-			Intermediates.DefineRefs[Index].value = Intermediates.Defines[Index].Value.GetData();
-		}
-	}
-
 	// Returns whether the specified line of text contains only these characters, making it a valid line marker from DXC: ' ', '\t', '~', '^'
 	static bool IsTextLineDxcLineMarker(const FString& Line)
 	{
@@ -442,7 +428,7 @@ namespace CrossCompiler
 		// Convert macro definitions map into an array container
 		if (Definitions != nullptr)
 		{
-			ConvertScDefines(*Intermediates, *Definitions);
+			ConvertStringMapToMacroDefines(Definitions->GetDefinitionMap(), Intermediates->Defines, Intermediates->DefineRefs);
 		}
 
 		// Convert shader stage
@@ -461,7 +447,7 @@ namespace CrossCompiler
 		// Convert macro definitions map into an array container
 		if (Definitions != nullptr)
 		{
-			ConvertScDefines(*Intermediates, *Definitions);
+			ConvertStringMapToMacroDefines(Definitions->GetDefinitionMap(), Intermediates->Defines, Intermediates->DefineRefs);
 		}
 
 		// Convert shader stage
