UE5 - Allow passing extra arguments to DXC

diff --git a/Engine/Source/Developer/ShaderCompilerCommon/Private/ShaderConductorContext.cpp b/Engine/Source/Developer/ShaderCompilerCommon/Private/ShaderConductorContext.cpp
--- a/Engine/Source/Developer/ShaderCompilerCommon/Private/ShaderConductorContext.cpp
+++ b/Engine/Source/Developer/ShaderCompilerCommon/Private/ShaderConductorContext.cpp
@@ -213,6 +213,8 @@ namespace CrossCompiler
 		TArray<ShaderConductor::MacroDefine> DefineRefs;
 		TArray<TPair<TArray<ANSICHAR>, TArray<ANSICHAR>>> Flags;
 		TArray<ShaderConductor::MacroDefine> FlagRefs;
+		TArray<TArray<ANSICHAR>> ExtraDxcArgs;
+		TArray<ANSICHAR const*> ExtraDxcArgRefs;
 	};
 
 	static void ConvertScSourceDesc(const FShaderConductorContext::FShaderConductorIntermediates& Intermediates, ShaderConductor::Compiler::SourceDesc& OutSourceDesc)
@@ -273,6 +275,25 @@ namespace CrossCompiler
 		checkf(OutTargetDesc.version, TEXT("Unsupported target shader version for Metal family: %d"), InTarget.Version);
 	}
 
+	// Converts an array of FString to a C-style array of char* pointers
+	static void ConvertStringArrayToAnsiArray(const TArray<FString>& InPairs, TArray<TArray<ANSICHAR>>& OutPairs, TArray<const char*>& OutPairRefs)
+	{
+		// Convert map into an array container
+		TArray<ANSICHAR> Value;
+		for (const FString& Iter : InPairs)
+		{
+			ConvertFStringToAnsiString(Iter, Value);
+			OutPairs.Emplace(MoveTemp(Value));
+		}
+
+		// Store references after all elements have been added to the container so the pointers remain valid
+		OutPairRefs.SetNum(OutPairs.Num());
+		for (int32 Index = 0; Index < OutPairs.Num(); ++Index)
+		{
+			OutPairRefs[Index] = OutPairs[Index].GetData();
+		}
+	}
+
 	// Converts a map of string pairs to a C-Style macro defines array
 	static void ConvertStringMapToMacroDefines(const TMap<FString,FString>& InPairs, TArray<TPair<TArray<ANSICHAR>, TArray<ANSICHAR>>>& OutPairs, TArray<ShaderConductor::MacroDefine>& OutPairRefs)
 	{
@@ -337,7 +358,7 @@ namespace CrossCompiler
 		}
 	}
 
-	static void ConvertScOptions(const FShaderConductorOptions& InOptions, ShaderConductor::Compiler::Options& OutOptions)
+	static void ConvertScOptions(const FShaderConductorOptions& InOptions, ShaderConductor::Compiler::Options& OutOptions, const TArray<ANSICHAR const*>& ExtraDxcArgRefs)
 	{
 		OutOptions.removeUnusedGlobals = InOptions.bRemoveUnusedGlobals;
 		OutOptions.packMatricesInRowMajor = InOptions.bPackMatricesInRowMajor;
@@ -354,6 +375,17 @@ namespace CrossCompiler
 		{
 			OutOptions.shaderModel = ToShaderConductorShaderModel(InOptions.TargetProfile);
 		}
+
+		if (ExtraDxcArgRefs.Num() > 0)
+		{
+			OutOptions.numDXCArgs = ExtraDxcArgRefs.Num();
+			OutOptions.DXCArgs = (const char**)ExtraDxcArgRefs.GetData();
+		}
+		else
+		{
+			OutOptions.numDXCArgs = 0;
+			OutOptions.DXCArgs = nullptr;
+		}
 	}
 
 	// Returns whether the specified line of text contains only these characters, making it a valid line marker from DXC: ' ', '\t', '~', '^'
@@ -418,7 +450,7 @@ namespace CrossCompiler
 		return *this;
 	}
 
-	bool FShaderConductorContext::LoadSource(const FString& ShaderSource, const FString& Filename, const FString& EntryPoint, EHlslShaderFrequency ShaderStage, const FShaderCompilerDefinitions* Definitions)
+	bool FShaderConductorContext::LoadSource(const FString& ShaderSource, const FString& Filename, const FString& EntryPoint, EHlslShaderFrequency ShaderStage, const FShaderCompilerDefinitions* Definitions, const TArray<FString>* ExtraDxcArgs)
 	{
 		// Convert FString to ANSI string and store them as intermediates
 		ConvertFStringToAnsiString(ShaderSource, Intermediates->ShaderSource);
@@ -431,13 +463,18 @@ namespace CrossCompiler
 			ConvertStringMapToMacroDefines(Definitions->GetDefinitionMap(), Intermediates->Defines, Intermediates->DefineRefs);
 		}
 
+		if (ExtraDxcArgs && ExtraDxcArgs->Num() > 0)
+		{
+			ConvertStringArrayToAnsiArray(*ExtraDxcArgs, Intermediates->ExtraDxcArgs, Intermediates->ExtraDxcArgRefs);
+		}
+
 		// Convert shader stage
 		Intermediates->Stage = ToShaderConductorShaderStage(ShaderStage);
 
 		return true;
 	}
 
-	bool FShaderConductorContext::LoadSource(const ANSICHAR* ShaderSource, const ANSICHAR* Filename, const ANSICHAR* EntryPoint, EHlslShaderFrequency ShaderStage, const FShaderCompilerDefinitions* Definitions)
+	bool FShaderConductorContext::LoadSource(const ANSICHAR* ShaderSource, const ANSICHAR* Filename, const ANSICHAR* EntryPoint, EHlslShaderFrequency ShaderStage, const FShaderCompilerDefinitions* Definitions, const TArray<FString>* ExtraDxcArgs)
 	{
 		// Store ANSI strings as intermediates
 		CopyAnsiString(ShaderSource, Intermediates->ShaderSource);
@@ -450,6 +487,11 @@ namespace CrossCompiler
 			ConvertStringMapToMacroDefines(Definitions->GetDefinitionMap(), Intermediates->Defines, Intermediates->DefineRefs);
 		}
 
+		if (ExtraDxcArgs && ExtraDxcArgs->Num() > 0)
+		{
+			ConvertStringArrayToAnsiArray(*ExtraDxcArgs, Intermediates->ExtraDxcArgs, Intermediates->ExtraDxcArgRefs);
+		}
+
 		// Convert shader stage
 		Intermediates->Stage = ToShaderConductorShaderStage(ShaderStage);
 
@@ -463,7 +505,7 @@ namespace CrossCompiler
 		ConvertScSourceDesc(*Intermediates, ScSourceDesc);
 
 		ShaderConductor::Compiler::Options ScOptions;
-		ConvertScOptions(Options, ScOptions);
+		ConvertScOptions(Options, ScOptions, Intermediates->ExtraDxcArgRefs);
 
 		// Rewrite HLSL with wrapper function to catch exceptions from ShaderConductor
 		bool bSucceeded = false;
@@ -509,7 +551,7 @@ namespace CrossCompiler
 		ScTargetDesc.language = ShaderConductor::ShadingLanguage::SpirV;
 
 		ShaderConductor::Compiler::Options ScOptions;
-		ConvertScOptions(Options, ScOptions);
+		ConvertScOptions(Options, ScOptions, Intermediates->ExtraDxcArgRefs);
 
 		// Compile HLSL source code to SPIR-V
 		bool bSucceeded = false;
@@ -577,7 +619,7 @@ namespace CrossCompiler
 		ConvertScTargetDesc(*Intermediates, Target, ScTargetDesc);
 
 		ShaderConductor::Compiler::Options ScOptions;
-		ConvertScOptions(Options, ScOptions);
+		ConvertScOptions(Options, ScOptions, Intermediates->ExtraDxcArgRefs);
 
 		ShaderConductor::Compiler::ResultDesc ScBinaryDesc;
 		ScBinaryDesc.target.Reset(InSpirv, InSpirvByteSize);
@@ -698,12 +740,12 @@ namespace CrossCompiler
 		return *this; // Dummy
 	}
 
-	bool FShaderConductorContext::LoadSource(const FString& ShaderSource, const FString& Filename, const FString& EntryPoint, EHlslShaderFrequency ShaderStage, const FShaderCompilerDefinitions* Definitions)
+	bool FShaderConductorContext::LoadSource(const FString& ShaderSource, const FString& Filename, const FString& EntryPoint, EHlslShaderFrequency ShaderStage, const FShaderCompilerDefinitions* Definitions, const TArray<FString>* ExtraDxcArgs)
 	{
 		return false; // Dummy
 	}
 
-	bool FShaderConductorContext::LoadSource(const ANSICHAR* ShaderSource, const ANSICHAR* Filename, const ANSICHAR* EntryPoint, EHlslShaderFrequency ShaderStage, const FShaderCompilerDefinitions* Definitions)
+	bool FShaderConductorContext::LoadSource(const ANSICHAR* ShaderSource, const ANSICHAR* Filename, const ANSICHAR* EntryPoint, EHlslShaderFrequency ShaderStage, const FShaderCompilerDefinitions* Definitions, const TArray<FString>* ExtraDxcArgs)
 	{
 		return false; // Dummy
 	}
diff --git a/Engine/Source/Developer/ShaderCompilerCommon/Public/ShaderConductorContext.h b/Engine/Source/Developer/ShaderCompilerCommon/Public/ShaderConductorContext.h
--- a/Engine/Source/Developer/ShaderCompilerCommon/Public/ShaderConductorContext.h
+++ b/Engine/Source/Developer/ShaderCompilerCommon/Public/ShaderConductorContext.h
@@ -32,6 +32,9 @@ namespace CrossCompiler
 		/** Enable a pass that converts floating point MUL+ADD pairs into FMAs to avoid re-association. */
 		bool bEnableFMAPass = false;
 
+		/** Pass-through command line options into DXC. */
+		TArray<FString> ExtraDxcOptions;
+
 		/** Target shader profile. By default HCT_FeatureLevelSM5. */
 		EHlslCompileTarget TargetProfile = HCT_FeatureLevelSM5;
 	};
@@ -85,8 +88,8 @@ namespace CrossCompiler
 		FShaderConductorContext& operator = (const FShaderConductorContext&) = delete;
 
 		/** Loads the shader source and converts the input descriptor to a format suitable for ShaderConductor. If 'Definitions' is null, the previously loaded definitions are not modified. */
-		bool LoadSource(const FString& ShaderSource, const FString& Filename, const FString& EntryPoint, EHlslShaderFrequency ShaderStage, const FShaderCompilerDefinitions* Definitions = nullptr);
-		bool LoadSource(const ANSICHAR* ShaderSource, const ANSICHAR* Filename, const ANSICHAR* EntryPoint, EHlslShaderFrequency ShaderStage, const FShaderCompilerDefinitions* Definitions = nullptr);
+		bool LoadSource(const FString& ShaderSource, const FString& Filename, const FString& EntryPoint, EHlslShaderFrequency ShaderStage, const FShaderCompilerDefinitions* Definitions = nullptr, const TArray<FString>* ExtraDxcArgs = nullptr);
+		bool LoadSource(const ANSICHAR* ShaderSource, const ANSICHAR* Filename, const ANSICHAR* EntryPoint, EHlslShaderFrequency ShaderStage, const FShaderCompilerDefinitions* Definitions = nullptr, const TArray<FString>* ExtraDxcArgs = nullptr);
 
 		/** Rewrites the specified HLSL shader source code. This allows to reduce the HLSL code by removing unused global resources for instance.
 		This will update the internally loaded source (see 'LoadSource'), so the output parameter 'OutSource' is optional. */
