Replaced "ArraySize" macro with template function to avoid identifier conflict.

#rnx

diff --git a/Engine/Source/Developer/ShaderCompilerCommon/Private/HlslLexer.cpp b/Engine/Source/Developer/ShaderCompilerCommon/Private/HlslLexer.cpp
--- a/Engine/Source/Developer/ShaderCompilerCommon/Private/HlslLexer.cpp
+++ b/Engine/Source/Developer/ShaderCompilerCommon/Private/HlslLexer.cpp
@@ -8,10 +8,7 @@
 
 namespace CrossCompiler
 {
-	template <typename T, size_t N>
-	char (&ArraySizeHelper(T (&array)[N]))[N];
-#define ArraySize(array) (sizeof(ArraySizeHelper(array)))
-#define MATCH_TARGET(S) S, (int32)ArraySize(S) - 1
+#define MATCH_TARGET(S) S, (UE_ARRAY_COUNT(S) - 1)
 
 	typedef FPlatformTypes::TCHAR TCHAR;
 
@@ -661,7 +658,7 @@ namespace CrossCompiler
 					// Check for infinity constant
 					if (Match('#'))
 					{
-						if (MatchString(TEXT("INF"), 3))
+						if (MatchString(MATCH_TARGET(TEXT("INF"))))
 						{
 							goto Done;
 						}
