Added secondary "WritePackedGlobal" function to FHlslccHeaderWriter.

#rnx

diff --git a/Engine/Source/Developer/ShaderCompilerCommon/Private/HlslccHeaderWriter.cpp b/Engine/Source/Developer/ShaderCompilerCommon/Private/HlslccHeaderWriter.cpp
--- a/Engine/Source/Developer/ShaderCompilerCommon/Private/HlslccHeaderWriter.cpp
+++ b/Engine/Source/Developer/ShaderCompilerCommon/Private/HlslccHeaderWriter.cpp
@@ -352,11 +352,43 @@ namespace CrossCompiler
 		MetaDataPrintf(Strings.UniformBlocks, TEXT("%s(%u)"), ResourceName, BindingIndex);
 	}
 
-	void FHlslccHeaderWriter::WritePackedGlobal(const TCHAR* ResourceName, const TCHAR* TypeSpecifier, uint32 ByteOffset, uint32 ByteSize)
+	static EPackedTypeName EncodePackedGlobalType(const SpvReflectTypeDescription& TypeDescription)
 	{
-		checkf(ByteOffset % 4 == 0, TEXT("field offset in @PackedGlobals shader meta data must be a multiple of 4, but got %u"), ByteOffset);
-		checkf(ByteSize % 4 == 0, TEXT("field size in @PackedGlobals shader meta data must be a multiple of 4, but got %u"), ByteSize);
-		MetaDataPrintf(Strings.PackedGlobals, TEXT("%s(%s:%u,%u)"), ResourceName, TypeSpecifier, ByteOffset / 4, ByteSize / 4);
+		const SpvReflectTypeFlags ScalarTypeFlagsBitmask =
+		(
+			SPV_REFLECT_TYPE_FLAG_VOID	|
+			SPV_REFLECT_TYPE_FLAG_BOOL	|
+			SPV_REFLECT_TYPE_FLAG_INT	|
+			SPV_REFLECT_TYPE_FLAG_FLOAT
+		);
+		const SpvReflectTypeFlags TypeFlags = TypeDescription.type_flags;
+		const SpvReflectTypeFlags MaskedType = TypeFlags & ScalarTypeFlagsBitmask;
+
+		switch (MaskedType)
+		{
+		case SPV_REFLECT_TYPE_FLAG_BOOL:
+		case SPV_REFLECT_TYPE_FLAG_INT:
+			return (TypeDescription.traits.numeric.scalar.signedness ? EPackedTypeName::Int : EPackedTypeName::Uint);
+		case SPV_REFLECT_TYPE_FLAG_FLOAT:
+			return EPackedTypeName::HighP;
+		default:
+			checkf(false, TEXT("unsupported component type %d"), MaskedType);
+			return EPackedTypeName::LowP;
+		}
+	}
+
+	void FHlslccHeaderWriter::WritePackedGlobal(const TCHAR* ResourceName, EPackedTypeName PackedType, uint32 ByteOffset, uint32 ByteSize)
+	{
+		// PackedType must be one of 'h','m','l','i','u'.
+		checkf(ByteOffset % 4 == 0, TEXT("field offset of \"%s\" in @PackedGlobals shader meta data must be a multiple of 4, but got %u"), ResourceName, ByteOffset);
+		checkf(ByteSize % 4 == 0, TEXT("field size of \"%s\" in @PackedGlobals shader meta data must be a multiple of 4, but got %u"), ResourceName, ByteSize);
+		MetaDataPrintf(Strings.PackedGlobals, TEXT("%s(%c:%u,%u)"), ResourceName, TCHAR(PackedType), ByteOffset / 4, ByteSize / 4);
+	}
+
+	void FHlslccHeaderWriter::WritePackedGlobal(const SpvReflectBlockVariable& Variable)
+	{
+		const EPackedTypeName PackedType = EncodePackedGlobalType(*(Variable.type_description));
+		WritePackedGlobal(ANSI_TO_TCHAR(Variable.name), PackedType, Variable.absolute_offset, Variable.size);// padded_size);
 	}
 
 	void FHlslccHeaderWriter::WritePackedUB(uint32 BindingIndex)
diff --git a/Engine/Source/Developer/ShaderCompilerCommon/Public/HlslccHeaderWriter.h b/Engine/Source/Developer/ShaderCompilerCommon/Public/HlslccHeaderWriter.h
--- a/Engine/Source/Developer/ShaderCompilerCommon/Public/HlslccHeaderWriter.h
+++ b/Engine/Source/Developer/ShaderCompilerCommon/Public/HlslccHeaderWriter.h
@@ -3,10 +3,12 @@
 #pragma once
 
 #include "CoreMinimal.h"
+#include "CrossCompilerCommon.h"
 
 
 // Forward declaration from <spirv_reflect.h>
 struct SpvReflectInterfaceVariable;
+struct SpvReflectBlockVariable;
 
 // Cross compiler support/common functionality
 namespace CrossCompiler
@@ -20,7 +22,8 @@ namespace CrossCompiler
 		void WriteOutputAttribute(const SpvReflectInterfaceVariable& Attribute);
 		void WriteOutputAttribute(const TCHAR* AttributeName, const TCHAR* TypeSpecifier, int32 Location, bool bLocationPrefix, bool bLocationSuffix);
 		void WriteUniformBlock(const TCHAR* ResourceName, uint32 BindingIndex);
-		void WritePackedGlobal(const TCHAR* ResourceName, const TCHAR* TypeSpecifier, uint32 ByteOffset, uint32 ByteSize);
+		void WritePackedGlobal(const TCHAR* ResourceName, EPackedTypeName PackedType, uint32 ByteOffset, uint32 ByteSize);
+		void WritePackedGlobal(const SpvReflectBlockVariable& Variable);
 		void WritePackedUB(uint32 BindingIndex);
 		void WritePackedUBField(const TCHAR* ResourceName, uint32 ByteOffset, uint32 ByteSize);
 		void WritePackedUBGlobalCopy(uint32 SourceCB, uint32 SourceOffset, uint32 DestCBIndex, uint32 DestCBPrecision, uint32 DestOffset, uint32 Size, bool bGroupFlattenedUBs = false);
diff --git a/Engine/Source/Developer/ShaderFormatOpenGL/Private/OpenGLShaderCompiler.cpp b/Engine/Source/Developer/ShaderFormatOpenGL/Private/OpenGLShaderCompiler.cpp
--- a/Engine/Source/Developer/ShaderFormatOpenGL/Private/OpenGLShaderCompiler.cpp
+++ b/Engine/Source/Developer/ShaderFormatOpenGL/Private/OpenGLShaderCompiler.cpp
@@ -1596,8 +1596,7 @@ static bool CompileToGlslWithShaderConductor(
 				{
 					for (uint32 i = 0; i < Binding->block.member_count; i++)
 					{
-						SpvReflectBlockVariable& member = Binding->block.members[i];
-						CCHeaderWriter.WritePackedGlobal(UTF8_TO_TCHAR(member.name), TEXT("h"), member.absolute_offset, member.size);
+						CCHeaderWriter.WritePackedGlobal(Binding->block.members[i]);
 					}
 				}
 				else
@@ -1686,7 +1685,7 @@ static bool CompileToGlslWithShaderConductor(
 
 							uint32& Offset = GlobalOffsets.FindOrAdd(TypeQualifier);
 
-							CCHeaderWriter.WritePackedGlobal(*MemberName, *TypeQualifier, member.absolute_offset, member.size);
+							CCHeaderWriter.WritePackedGlobal(member);
 
 							bool const bArray = type.traits.array.dims_count > 0;
 
diff --git a/Engine/Source/Developer/VulkanShaderFormat/Private/VulkanShaderCompiler.cpp b/Engine/Source/Developer/VulkanShaderFormat/Private/VulkanShaderCompiler.cpp
--- a/Engine/Source/Developer/VulkanShaderFormat/Private/VulkanShaderCompiler.cpp
+++ b/Engine/Source/Developer/VulkanShaderFormat/Private/VulkanShaderCompiler.cpp
@@ -2256,7 +2256,6 @@ static void BuildShaderOutputFromSpirv(
 	for (const SpvReflectDescriptorBinding* Binding : Bindings.UniformBuffers)
 	{
 		const FString ResourceName(ANSI_TO_TCHAR(Binding->name));
-		const uint32 SizePerComponent = sizeof(float);
 
 		if (ResourceName == UBOGlobalsNameSpv)
 		{
@@ -2276,8 +2275,7 @@ static void BuildShaderOutputFromSpirv(
 			for (uint32 MemberIndex = 0; MemberIndex < Binding->block.member_count; ++MemberIndex)
 			{
 				const SpvReflectBlockVariable* Member = &(Binding->block.members[MemberIndex]);
-				const FString MemberName(ANSI_TO_TCHAR(Member->name));
-				CCHeaderWriter.WritePackedGlobal(*MemberName, TEXT("h"), Member->absolute_offset, Member->size);
+				CCHeaderWriter.WritePackedGlobal(ANSI_TO_TCHAR(Member->name), CrossCompiler::EPackedTypeName::HighP, Member->absolute_offset, Member->size);
 			}
 
 			// Stop after we found $Globals uniform buffer
@@ -2288,7 +2286,6 @@ static void BuildShaderOutputFromSpirv(
 	for (const SpvReflectDescriptorBinding* Binding : Bindings.UniformBuffers)
 	{
 		const FString ResourceName(ANSI_TO_TCHAR(Binding->name));
-		const uint32 SizePerComponent = sizeof(float);
 
 		if (ResourceName != UBOGlobalsNameSpv)
 		{
