Fix for Niagara velocity shaders being generated incorrectly, which would fail PSOs to be created.
Fixes issue where an array of non-indexed elements with custom semantics would have their 'optimized' code path generated as:
Local_Value = OptimizedValue;
Instead of the reverse.

diff --git a/Engine/Shaders/Public/ShaderVersion.ush b/Engine/Shaders/Public/ShaderVersion.ush
--- a/Engine/Shaders/Public/ShaderVersion.ush
+++ b/Engine/Shaders/Public/ShaderVersion.ush
@@ -2,4 +2,4 @@
 // Each time the console command is executed it generates a new GUID. As this file is included
 // in Platform.ush (which should be included in any shader) it allows to invalidate the shader DDC.
 // 
-// GUID = DDCF4C7AB9BB4D43BB2F9D07CD995890
\ No newline at end of file
+// GUID = 5BFC29C6A26348CCBC8D001186A1272A
\ No newline at end of file
diff --git a/Engine/Source/Developer/ShaderCompilerCommon/Private/HlslUtils.cpp b/Engine/Source/Developer/ShaderCompilerCommon/Private/HlslUtils.cpp
--- a/Engine/Source/Developer/ShaderCompilerCommon/Private/HlslUtils.cpp
+++ b/Engine/Source/Developer/ShaderCompilerCommon/Private/HlslUtils.cpp
@@ -810,7 +810,7 @@ struct FRemoveUnusedOutputs : FRemoveAlgorithm
 								NewDeclaratorList->Declarations.Add(MemberDeclaration);
 								DestStruct->Members.Add(NewDeclaratorList);
 
-								CopyMember(MemberDeclaration, SourcePrefix, DestPrefix, BodyContext.PostInstructions);
+								CopyMember(MemberDeclaration, DestPrefix, SourcePrefix, BodyContext.PostInstructions);
 							}
 
 							for (uint32 Index = 0; Index < ArrayLength; ++Index)
@@ -1279,7 +1279,7 @@ struct FRemoveUnusedInputs : FRemoveAlgorithm
 								NewDeclaratorList->Declarations.Add(MemberDeclaration);
 								DestStruct->Members.Add(NewDeclaratorList);
 
-								CopyMember(MemberDeclaration, SourcePrefix, DestPrefix, BodyContext.PreInstructions);
+								CopyMember(MemberDeclaration, DestPrefix, SourcePrefix, BodyContext.PreInstructions);
 							}
 
 							for (uint32 Index = 0; Index < ArrayLength; ++Index)
