Replace custom Disassemble function in ShaderConductor interface with the one that was already there.

#rnx

diff --git a/Engine/Source/Developer/ShaderCompilerCommon/Private/ShaderCompilerCommon.cpp b/Engine/Source/Developer/ShaderCompilerCommon/Private/ShaderCompilerCommon.cpp
--- a/Engine/Source/Developer/ShaderCompilerCommon/Private/ShaderCompilerCommon.cpp
+++ b/Engine/Source/Developer/ShaderCompilerCommon/Private/ShaderCompilerCommon.cpp
@@ -1581,12 +1581,12 @@ void DumpDebugShaderBinary(const FShaderCompilerInput& Input, void* InData, int3
 	}
 }
 
-void DumpDebugShaderDisassembledSpirv(const FShaderCompilerInput& Input, void* InData, int32 InDataByteSize, const FString& FileExtension)
+static void DumpDebugShaderDisassembled(const FShaderCompilerInput& Input, CrossCompiler::EShaderConductorIR Language, void* InData, int32 InDataByteSize, const FString& FileExtension)
 {
 	if (InData != nullptr && InDataByteSize > 0 && !FileExtension.IsEmpty())
 	{
 		TArray<ANSICHAR> AssemblyText;
-		if (CrossCompiler::FShaderConductorContext::DisassembleSpirv(InData, InDataByteSize, AssemblyText))
+		if (CrossCompiler::FShaderConductorContext::Disassemble(Language, InData, InDataByteSize, AssemblyText))
 		{
 			// Assembly text contains NUL terminator, so text lenght is |array|-1
 			DumpDebugShaderText(Input, AssemblyText.GetData(), AssemblyText.Num() - 1, FileExtension);
@@ -1594,6 +1594,16 @@ void DumpDebugShaderDisassembledSpirv(const FShaderCompilerInput& Input, void* I
 	}
 }
 
+void DumpDebugShaderDisassembledSpirv(const FShaderCompilerInput& Input, void* InData, int32 InDataByteSize, const FString& FileExtension)
+{
+	DumpDebugShaderDisassembled(Input, CrossCompiler::EShaderConductorIR::Spirv, InData, InDataByteSize, FileExtension);
+}
+
+void DumpDebugShaderDisassembledDxil(const FShaderCompilerInput& Input, void* InData, int32 InDataByteSize, const FString& FileExtension)
+{
+	DumpDebugShaderDisassembled(Input, CrossCompiler::EShaderConductorIR::Dxil, InData, InDataByteSize, FileExtension);
+}
+
 namespace CrossCompiler
 {
 	FString CreateResourceTableFromEnvironment(const FShaderCompilerEnvironment& Environment)
diff --git a/Engine/Source/Developer/ShaderCompilerCommon/Private/ShaderConductorContext.cpp b/Engine/Source/Developer/ShaderCompilerCommon/Private/ShaderConductorContext.cpp
--- a/Engine/Source/Developer/ShaderCompilerCommon/Private/ShaderConductorContext.cpp
+++ b/Engine/Source/Developer/ShaderCompilerCommon/Private/ShaderConductorContext.cpp
@@ -762,22 +762,31 @@ namespace CrossCompiler
 		}
 	}
 
-	bool FShaderConductorContext::DisassembleSpirv(const void* InSpirv, uint32 InSpirvByteSize, TArray<ANSICHAR>& OutAssemblyText)
+	bool FShaderConductorContext::Disassemble(EShaderConductorIR Language, const void* Binary, uint32 BinaryByteSize, TArray<ANSICHAR>& OutAssemblyText)
 	{
 		// Initialize Blob with input SPIR-V code
-		ShaderConductor::Compiler::ResultDesc BinaryInput;
-		BinaryInput.isText = false;
-		BinaryInput.hasError = false;
-		BinaryInput.target = ShaderConductor::Blob(InSpirv, InSpirvByteSize);
+		ShaderConductor::Compiler::DisassembleDesc BinaryDesc;
+		switch (Language)
+		{
+		case EShaderConductorIR::Spirv:
+			BinaryDesc.language = ShaderConductor::ShadingLanguage::SpirV;
+			break;
+		case EShaderConductorIR::Dxil:
+			BinaryDesc.language = ShaderConductor::ShadingLanguage::Dxil;
+			break;
+		}
+		BinaryDesc.binary = reinterpret_cast<const uint8_t*>(Binary);
+		BinaryDesc.binarySize = BinaryByteSize;
 
 		// Disassemble via ShaderConductor interface
-		ShaderConductor::Compiler::ResultDesc TextOutput = ShaderConductor::Compiler::Disassemble(BinaryInput);
+		ShaderConductor::Compiler::ResultDesc TextOutput = ShaderConductor::Compiler::Disassemble(BinaryDesc);
 		if (TextOutput.isText && !TextOutput.hasError)
 		{
 			// Convert and return output to ANSI string
 			ConvertByteArrayToAnsiString(reinterpret_cast<const ANSICHAR*>(TextOutput.target.Data()), TextOutput.target.Size(), OutAssemblyText);
 			return true;
 		}
+
 		return false;
 	}
 
@@ -858,7 +867,7 @@ namespace CrossCompiler
 		// Dummy
 	}
 
-	bool FShaderConductorContext::DisassembleSpirv(const void* InSpirv, uint32 InSpirvByteSize, TArray<ANSICHAR>& OutAssemblyText)
+	bool FShaderConductorContext::Disassemble(EShaderConductorIR Language, const void* Binary, uint32 BinaryByteSize, TArray<ANSICHAR>& OutAssemblyText)
 	{
 		return false; // Dummy
 	}
diff --git a/Engine/Source/Developer/ShaderCompilerCommon/Public/ShaderCompilerCommon.h b/Engine/Source/Developer/ShaderCompilerCommon/Public/ShaderCompilerCommon.h
--- a/Engine/Source/Developer/ShaderCompilerCommon/Public/ShaderCompilerCommon.h
+++ b/Engine/Source/Developer/ShaderCompilerCommon/Public/ShaderCompilerCommon.h
@@ -113,6 +113,7 @@ extern SHADERCOMPILERCOMMON_API void DumpDebugShaderText(const FShaderCompilerIn
 extern SHADERCOMPILERCOMMON_API void DumpDebugShaderText(const FShaderCompilerInput& Input, ANSICHAR* InSource, int32 InSourceLength, const FString& FileExtension);
 extern SHADERCOMPILERCOMMON_API void DumpDebugShaderBinary(const FShaderCompilerInput& Input, void* InData, int32 InDataByteSize, const FString& FileExtension);
 extern SHADERCOMPILERCOMMON_API void DumpDebugShaderDisassembledSpirv(const FShaderCompilerInput& Input, void* InData, int32 InDataByteSize, const FString& FileExtension);
+extern SHADERCOMPILERCOMMON_API void DumpDebugShaderDisassembledDxil(const FShaderCompilerInput& Input, void* InData, int32 InDataByteSize, const FString& FileExtension);
 
 UE_DEPRECATED(4.26, "SourceLength is no longer needed.")
 inline void DumpDebugUSF(const FShaderCompilerInput& Input, const ANSICHAR* Source, int32 SourceLength, uint32 HlslCCFlags = 0, const TCHAR* OverrideBaseFilename = nullptr)
diff --git a/Engine/Source/Developer/ShaderCompilerCommon/Public/ShaderConductorContext.h b/Engine/Source/Developer/ShaderCompilerCommon/Public/ShaderConductorContext.h
--- a/Engine/Source/Developer/ShaderCompilerCommon/Public/ShaderConductorContext.h
+++ b/Engine/Source/Developer/ShaderCompilerCommon/Public/ShaderConductorContext.h
@@ -48,6 +48,13 @@ namespace CrossCompiler
 		Metal_iOS,
 	};
 
+	/** Intermediate representation languages for ShaderConductor disassembly output. */
+	enum class EShaderConductorIR
+	{
+		Spirv,
+		Dxil,
+	};
+
 	/** Shader conductor output target descriptor. */
 	struct SHADERCOMPILERCOMMON_API FShaderConductorTarget
 	{
@@ -150,7 +157,7 @@ namespace CrossCompiler
 		static const FShaderConductorIdentifierTable& GetIdentifierTable();
 
 		/** Disassembles the specified SPIR-V module and returns its assembly as text representation. */
-		static bool DisassembleSpirv(const void* InSpirv, uint32 InSpirvByteSize, TArray<ANSICHAR>& OutAssemblyText);
+		static bool Disassemble(EShaderConductorIR Language, const void* Binary, uint32 BinaryByteSize, TArray<ANSICHAR>& OutAssemblyText);
 
 	public:
 		struct FShaderConductorIntermediates; // Pimpl idiom
