Improve assertion output for redundancy check in SPIR-V header parser.

#rnx

diff --git a/Engine/Source/Developer/ShaderCompilerCommon/Public/SpirvCommon.h b/Engine/Source/Developer/ShaderCompilerCommon/Public/SpirvCommon.h
--- a/Engine/Source/Developer/ShaderCompilerCommon/Public/SpirvCommon.h
+++ b/Engine/Source/Developer/ShaderCompilerCommon/Public/SpirvCommon.h
@@ -37,7 +37,7 @@ public:
 			//const uint32 Generator = Ptr[2];
 			//const uint32 Bound = Ptr[3];
 			const uint32 ZeroCheckValue = Ptr[4];
-			check(ZeroCheckValue == 0);
+			checkf(ZeroCheckValue == 0, TEXT("redundancy check for SPIR-V module failed: WORD[4] is non-zero"));
 
 			// Move pointer to first instruction
 			Ptr += 5;
