Remove unused struct field from "FShaderConductorOptions".

#rnx

diff --git a/Engine/Source/Developer/ShaderCompilerCommon/Public/ShaderConductorContext.h b/Engine/Source/Developer/ShaderCompilerCommon/Public/ShaderConductorContext.h
--- a/Engine/Source/Developer/ShaderCompilerCommon/Public/ShaderConductorContext.h
+++ b/Engine/Source/Developer/ShaderCompilerCommon/Public/ShaderConductorContext.h
@@ -32,9 +32,6 @@ namespace CrossCompiler
 		/** Enable a pass that converts floating point MUL+ADD pairs into FMAs to avoid re-association. */
 		bool bEnableFMAPass = false;
 
-		/** Pass-through command line options into DXC. */
-		TArray<FString> ExtraDxcOptions;
-
 		/** Target shader profile. By default HCT_FeatureLevelSM5. */
 		EHlslCompileTarget TargetProfile = HCT_FeatureLevelSM5;
 	};
