Add proper ASCII texture to print standard character set from shader.

diff --git a/Engine/Shaders/Private/MiniFontCommon.ush b/Engine/Shaders/Private/MiniFontCommon.ush
--- a/Engine/Shaders/Private/MiniFontCommon.ush
+++ b/Engine/Shaders/Private/MiniFontCommon.ush
@@ -10,55 +10,57 @@ Texture2D MiniFontTexture;
 
 #define NATIVE_CHARACTER_RES 8
 
-// CharacterIDs for the MiniFont
-#define _0_ 0
-#define _1_ 1
-#define _2_ 2
-#define _3_ 3
-#define _4_ 4
-#define _5_ 5
-#define _6_ 6
-#define _7_ 7
-#define _8_ 8
-#define _9_ 9
-#define _A_ 10
-#define _B_ 11
-#define _C_ 12
-#define _D_ 13
-#define _E_ 14
-#define _F_ 15
-#define _G_ 16
-#define _H_ 17
-#define _I_ 18
-#define _J_ 19
-#define _K_ 20
-#define _L_ 21
-#define _M_ 22
-#define _N_ 23
-#define _O_ 24
-#define _P_ 25
-#define _Q_ 26
-#define _R_ 27
-#define _S_ 28
-#define _T_ 29
-#define _U_ 30
-#define _V_ 31
-#define _W_ 32
-#define _X_ 33
-#define _Y_ 34
-#define _Z_ 35
-#define _MINUS_ 59
-#define _COMMA_ 60
-#define _DOT_ 61
-#define _PLUS_ 62
-#define _SPC_ 63
+// Sub-set of ASCII characters for writing withing shader
+#define _0_ 48
+#define _1_ 49
+#define _2_ 50
+#define _3_ 51
+#define _4_ 52
+#define _5_ 53
+#define _6_ 54
+#define _7_ 55
+#define _8_ 56
+#define _9_ 57
+#define _A_ 65
+#define _B_ 66
+#define _C_ 67
+#define _D_ 68
+#define _E_ 69
+#define _F_ 70
+#define _G_ 71
+#define _H_ 72
+#define _I_ 73
+#define _J_ 74
+#define _K_ 75
+#define _L_ 76
+#define _M_ 77
+#define _N_ 78
+#define _O_ 79
+#define _P_ 80
+#define _Q_ 81
+#define _R_ 82
+#define _S_ 83
+#define _T_ 84
+#define _U_ 85
+#define _V_ 86
+#define _W_ 87
+#define _X_ 88
+#define _Y_ 89
+#define _Z_ 90
+#define _MINUS_ 45
+#define _COMMA_ 44
+#define _DOT_ 46
+#define _PLUS_ 43
+#define _SPC_ 32
 
 // sample from the MiniFont texture
 // @param CharacterID - the character to print
 // @param Position - the position within the 8x8 character
-float SampleMiniFont(int CharacterID, uint2 Position)
+float SampleMiniFont(int InAsciiCode, uint2 Position)
 {
-	return MiniFontTexture.Load(int3(CharacterID * 8 + Position.x, Position.y, 0)).r;
+	// Limit ASCII character to the Standard character set (32 - 127)
+	const uint TextureCode = clamp(InAsciiCode, 32, 127) - 32;
+	return MiniFontTexture.Load(int3(TextureCode * NATIVE_CHARACTER_RES + Position.x, Position.y, 0)).r;
 }
 
 // for printf debugging in the shader
diff --git a/Engine/Source/Developer/ShaderCompilerCommon/Private/ShaderCompilerCommon.cpp b/Engine/Source/Developer/ShaderCompilerCommon/Private/ShaderCompilerCommon.cpp
--- a/Engine/Source/Developer/ShaderCompilerCommon/Private/ShaderCompilerCommon.cpp
+++ b/Engine/Source/Developer/ShaderCompilerCommon/Private/ShaderCompilerCommon.cpp
@@ -1227,7 +1227,7 @@ void RemoveUniformBuffersFromSource(const FShaderCompilerEnvironment& Environmen
 }
 
 ///////////////////////////////////////////////////////////////////////////////////////////////////////////////////////////////////////
-// Process TEXT() macro to convert them into GPU MiniFont characters
+// Process TEXT() macro to convert them into GPU ASCII characters
 
 FString ParseText(const TCHAR* StartPtr, const TCHAR*& EndPtr)
 {
@@ -1251,113 +1251,18 @@ FString ParseText(const TCHAR* StartPtr, const TCHAR*& EndPtr)
 	return Out;
 }
 
-void ConvertTextToMiniFontCharacter(const FString& InText, FString& OutText, FString& OutEncodedText)
+void ConvertTextToAsciiCharacter(const FString& InText, FString& OutText, FString& OutEncodedText)
 {
-	// From MiniFontCommon.ush - CharacterIDs for the MiniFont
-	// Instead of replicating these values here, we could add them as DEFINE, to insure CPP/Shader are in sync.
-	// TODO: change this to proper ASCII code, and use an ASCII table on the shader side
-	const int32 _0_ = 0;
-	const int32 _1_ = 1;
-	const int32 _2_ = 2;
-	const int32 _3_ = 3;
-	const int32 _4_ = 4;
-	const int32 _5_ = 5;
-	const int32 _6_ = 6;
-	const int32 _7_ = 7;
-	const int32 _8_ = 8;
-	const int32 _9_ = 9;
-	const int32 _A_ = 10;
-	const int32 _B_ = 11;
-	const int32 _C_ = 12;
-	const int32 _D_ = 13;
-	const int32 _E_ = 14;
-	const int32 _F_ = 15;
-	const int32 _G_ = 16;
-	const int32 _H_ = 17;
-	const int32 _I_ = 18;
-	const int32 _J_ = 19;
-	const int32 _K_ = 20;
-	const int32 _L_ = 21;
-	const int32 _M_ = 22;
-	const int32 _N_ = 23;
-	const int32 _O_ = 24;
-	const int32 _P_ = 25;
-	const int32 _Q_ = 26;
-	const int32 _R_ = 27;
-	const int32 _S_ = 28;
-	const int32 _T_ = 29;
-	const int32 _U_ = 30;
-	const int32 _V_ = 31;
-	const int32 _W_ = 32;
-	const int32 _X_ = 33;
-	const int32 _Y_ = 34;
-	const int32 _Z_ = 35;
-	const int32 _MINUS_ = 59;
-	const int32 _COMMA_ = 60;
-	const int32 _DOT_ = 61;
-	const int32 _PLUS_ = 62;
-	const int32 _SPC_ = 63;
-
-	// Trim text to valid character only
-	FString ValidText;
-	{
-		FString Text = InText.ToUpper();
-		const uint32 CharCount = Text.Len();
-
-		ValidText.Reserve(CharCount);
-		for (uint32 CharIt = 0; CharIt < CharCount; ++CharIt)
-		{
-			const char C = Text[CharIt];
-			if ((C >= '0' && C <= '9') || (C >= 'A' && C <= 'Z') || (C == '-') || (C == ',') || (C == '.') || (C == '+') || (C == ' '))
-			{
-				ValidText += C;
-			}
-		}
-	}
-
-	// Translate character into MiniFont encoding
+	const uint32 CharCount = InText.Len();
+	OutEncodedText.Reserve(CharCount * 3); // ~2 digits per character + a comma
+	OutText = InText;
+	for (uint32 CharIt = 0; CharIt < CharCount; ++CharIt)
 	{
-		const uint32 CharCount = ValidText.Len();
-		OutEncodedText.Reserve(CharCount * 3); // ~2 digits per character + a comma
-		OutText = ValidText;
-
-		for (uint32 CharIt = 0; CharIt < CharCount; ++CharIt)
+		const char C = InText[CharIt];
+		OutEncodedText.AppendInt(uint8(C));
+		if (CharIt + 1 != CharCount)
 		{
-			const char C = ValidText[CharIt];
-			if ((C >= '0' && C <= '9'))
-			{
-				const int32 Offset = C - '0';
-				OutEncodedText.AppendInt(_0_ + Offset);
-			}
-			if ((C >= 'A' && C <= 'Z'))
-			{
-				const int32 Offset = C - 'A';
-				OutEncodedText.AppendInt(_A_ + Offset);
-			}
-			else if (C == '-')
-			{
-				OutEncodedText.AppendInt(_MINUS_);
-			}
-			else if (C == ',')
-			{
-				OutEncodedText.AppendInt(_COMMA_);
-			}
-			else if (C == '.')
-			{
-				OutEncodedText.AppendInt(_DOT_);
-			}
-			else if (C == '+')
-			{
-				OutEncodedText.AppendInt(_PLUS_);
-			}
-			else if (C == ' ')
-			{
-				OutEncodedText.AppendInt(_SPC_);
-			}
-			if (CharIt + 1 != CharCount)
-			{
-				OutEncodedText += ',';
-			}
+			OutEncodedText += ',';
 		}
 	}
 }
@@ -1381,7 +1286,7 @@ void TransformStringIntoCharacterArray(FString& PreprocessedShaderSource)
 	// 3. Replace TEXT by its entry number
 	uint32 GlobalCount = 0;
 	{
-		const TCHAR* TextIdentifier = TEXT("TEXT(\"");
+		const TCHAR* TextIdentifier = TEXT("TEXT(");
 		const TCHAR* SearchPtr = FCString::Strstr(&PreprocessedShaderSource[0], TextIdentifier);
 		while (SearchPtr)
 		{
@@ -1400,7 +1305,7 @@ void TransformStringIntoCharacterArray(FString& PreprocessedShaderSource)
 				Entry.Index			= EntryIndex;
 				Entry.Offset		= GlobalCount;
 				Entry.SourceText	= Text;
-			    ConvertTextToMiniFontCharacter(Entry.SourceText, Entry.ConvertedText, Entry.EncodedText);
+				ConvertTextToAsciiCharacter(Entry.SourceText, Entry.ConvertedText, Entry.EncodedText);
 				Entry.Hash			= CityHash32((const char*)&Entry.SourceText.GetCharArray(), sizeof(FString::ElementType) * Entry.SourceText.Len());
 
 				GlobalCount += Entry.ConvertedText.Len();
@@ -1413,7 +1318,7 @@ void TransformStringIntoCharacterArray(FString& PreprocessedShaderSource)
 				PreprocessedShaderSource.InsertAt(StartIndex, FString::FromInt(EntryIndex));
 
 				// Update SearchPtr, as PreprocessedShaderSource has been modified, and its memory could have been reallocated, causing SearchPtr to be invalid.
-				SearchPtr = &PreprocessedShaderSource[0] + (StartIndex + CharCount);
+				SearchPtr = &PreprocessedShaderSource[0] + StartIndex;
 			}
 			SearchPtr = FCString::Strstr(SearchPtr, TextIdentifier);
 		}
diff --git a/Engine/Source/Runtime/Renderer/Private/ScreenPass.cpp b/Engine/Source/Runtime/Renderer/Private/ScreenPass.cpp
--- a/Engine/Source/Runtime/Renderer/Private/ScreenPass.cpp
+++ b/Engine/Source/Runtime/Renderer/Private/ScreenPass.cpp
@@ -15,9 +15,9 @@ RENDERER_API const FScreenTransform FScreenTransform::ViewportUVToScreenPos(FVec
 
 const FTextureRHIRef& GetMiniFontTexture()
 {
-	if (GEngine->MiniFontTexture)
+	if (GSystemTextures.AsciiTexture)
 	{
-		return GEngine->MiniFontTexture->GetResource()->TextureRHI;
+		return GSystemTextures.AsciiTexture->GetRenderTargetItem().TargetableTexture;
 	}
 	else
 	{
diff --git a/Engine/Source/Runtime/Renderer/Private/ShaderPrint.cpp b/Engine/Source/Runtime/Renderer/Private/ShaderPrint.cpp
--- a/Engine/Source/Runtime/Renderer/Private/ShaderPrint.cpp
+++ b/Engine/Source/Runtime/Renderer/Private/ShaderPrint.cpp
@@ -318,7 +318,7 @@ namespace ShaderPrint
 		// Non graph managed resources
 		FUniformBufferRef UniformBuffer = CreateUniformBuffer(View);
 		FRDGBufferSRVRef ValuesBuffer = GraphBuilder.CreateSRV(View.ShaderPrintValueBuffer);
-		FTextureRHIRef FontTexture = GEngine->MiniFontTexture != nullptr ? GEngine->MiniFontTexture->GetResource()->TextureRHI : GSystemTextures.BlackDummy->GetRenderTargetItem().ShaderResourceTexture;;
+		FTextureRHIRef FontTexture = GSystemTextures.AsciiTexture->GetRenderTargetItem().ShaderResourceTexture;
 
 		const ERHIFeatureLevel::Type FeatureLevel = View.GetFeatureLevel();
 		FGlobalShaderMap* GlobalShaderMap = GetGlobalShaderMap(FeatureLevel);
diff --git a/Engine/Source/Runtime/Renderer/Private/SkyAtmosphereRendering.cpp b/Engine/Source/Runtime/Renderer/Private/SkyAtmosphereRendering.cpp
--- a/Engine/Source/Runtime/Renderer/Private/SkyAtmosphereRendering.cpp
+++ b/Engine/Source/Runtime/Renderer/Private/SkyAtmosphereRendering.cpp
@@ -1956,7 +1956,7 @@ void FSceneRenderer::RenderSkyAtmosphereEditorNotifications(FRDGBuilder& GraphBu
 			RenderSkyAtmosphereEditorHudPS::FParameters* PassParameters = GraphBuilder.AllocParameters<RenderSkyAtmosphereEditorHudPS::FParameters>();
 			PassParameters->RenderTargets[0] = FRenderTargetBinding(SceneColorTexture, ERenderTargetLoadAction::ELoad);
 			PassParameters->ViewUniformBuffer = View.ViewUniformBuffer;
-			PassParameters->MiniFontTexture = GEngine->MiniFontTexture ? GEngine->MiniFontTexture->GetResource()->TextureRHI : GSystemTextures.WhiteDummy->GetRenderTargetItem().TargetableTexture;
+			PassParameters->MiniFontTexture = GSystemTextures.AsciiTexture->GetRenderTargetItem().ShaderResourceTexture;
 
 			FPixelShaderUtils::AddFullscreenPass<RenderSkyAtmosphereEditorHudPS>(GraphBuilder, View.ShaderMap, RDG_EVENT_NAME("SkyAtmosphereEditor"), PixelShader, PassParameters, View.ViewRect);
 		}
diff --git a/Engine/Source/Runtime/Renderer/Private/SystemTextures.cpp b/Engine/Source/Runtime/Renderer/Private/SystemTextures.cpp
--- a/Engine/Source/Runtime/Renderer/Private/SystemTextures.cpp
+++ b/Engine/Source/Runtime/Renderer/Private/SystemTextures.cpp
@@ -273,6 +273,69 @@ void FSystemTextures::InitializeFeatureLevelDependentTextures(FRHICommandListImm
 		HairLUT2 = HairLUT0;
 	}
 
+	// ASCII texture
+	{
+		EPixelFormat Format = PF_R8;
+		const uint32 BytesPerPixel	= 1;
+		const uint32 CharacterCount = 96u;
+		const uint32 CharacterRes	= 8u;
+		const uint32 CharacterStride= CharacterRes * BytesPerPixel;
+
+		FPooledRenderTargetDesc Desc(FPooledRenderTargetDesc::Create2DDesc(FIntPoint(CharacterCount * CharacterRes, CharacterRes), Format, FClearValueBinding::None, TexCreate_None, TexCreate_ShaderResource, false));
+		
+		const uint32 Characters[] =
+		{
+			0x00000000,0x00000000,0x0c1e1e0c,0x000c000c,0x00363636,0x00000000,0x367f3636,0x0036367f,
+			0x1e033e0c,0x000c1f30,0x18336300,0x0063660c,0x6e1c361c,0x006e333b,0x00030606,0x00000000,
+			0x06060c18,0x00180c06,0x18180c06,0x00060c18,0xff3c6600,0x0000663c,0x3f0c0c00,0x00000c0c,
+			0x00000000,0x060c0e00,0x3f000000,0x00000000,0x00000000,0x000c0c00,0x0c183060,0x00010306,
+			0x3f3b331e,0x001e3337,0x0c0c0f0c,0x003f0c0c,0x1c30331e,0x003f3306,0x1c30331e,0x001e3330,
+			0x33363c38,0x0030307f,0x301f033f,0x001e3330,0x1f03061c,0x001e3333,0x1830333f,0x0006060c,
+			0x1e33331e,0x001e3333,0x3e33331e,0x000e1830,0x0c0c0000,0x000c0c00,0x0c0c0000,0x060c0e00,
+			0x03060c18,0x00180c06,0x003f0000,0x0000003f,0x30180c06,0x00060c18,0x1830331e,0x000c000c,
+			0x7b7b633e,0x001e037b,0x33331e0c,0x0033333f,0x3e66663f,0x003f6666,0x0303663c,0x003c6603,
+			0x6666363f,0x003f3666,0x1e16467f,0x007f4616,0x1e16467f,0x000f0616,0x0303663c,0x007c6673,
+			0x3f333333,0x00333333,0x0c0c0c1e,0x001e0c0c,0x30303078,0x001e3333,0x1e366667,0x00676636,
+			0x0606060f,0x007f6646,0x6b7f7763,0x00636363,0x7b6f6763,0x00636373,0x6363361c,0x001c3663,
+			0x3e66663f,0x000f0606,0x3333331e,0x00381e3b,0x3e66663f,0x0067361e,0x1c07331e,0x001e3338,
+			0x0c0c2d3f,0x001e0c0c,0x33333333,0x003f3333,0x33333333,0x000c1e33,0x6b636363,0x0063777f,
+			0x1c366363,0x00636336,0x1e333333,0x001e0c0c,0x0c19337f,0x007f6346,0x0606061e,0x001e0606,
+			0x180c0603,0x00406030,0x1818181e,0x001e1818,0x63361c08,0x00000000,0x00000000,0xff000000,
+			0x00180c0c,0x00000000,0x301e0000,0x006e333e,0x663e0607,0x003d6666,0x331e0000,0x001e3303,
+			0x3e303038,0x006e3333,0x331e0000,0x001e033f,0x0f06361c,0x000f0606,0x336e0000,0x1f303e33,
+			0x6e360607,0x00676666,0x0c0e000c,0x001e0c0c,0x181e0018,0x0e1b1818,0x36660607,0x0067361e,
+			0x0c0c0c0e,0x001e0c0c,0x7f370000,0x0063636b,0x331f0000,0x00333333,0x331e0000,0x001e3333,
+			0x663b0000,0x0f063e66,0x336e0000,0x78303e33,0x361b0000,0x000f0636,0x033e0000,0x001f301e,
+			0x0c3e0c08,0x00182c0c,0x33330000,0x006e3333,0x33330000,0x000c1e33,0x63630000,0x00367f6b,
+			0x36630000,0x0063361c,0x33330000,0x1f303e33,0x193f0000,0x003f260c,0x070c0c38,0x00380c0c,
+			0x00181818,0x00181818,0x380c0c07,0x00070c0c,0x00003b6e,0x00000000,0x00000000,0x00000000
+		};
+
+		GRenderTargetPool.FindFreeElement(RHICmdList, Desc, AsciiTexture, TEXT("AsciiTexture"), ERenderTargetTransience::NonTransient);
+		uint32 DestStride;
+		uint8* DestBuffer = (uint8*)RHICmdList.LockTexture2D((FTexture2DRHIRef&)AsciiTexture->GetRenderTargetItem().ShaderResourceTexture, 0, RLM_WriteOnly, DestStride, false);
+
+		for (int32 c = 0; c < CharacterCount; c++)
+		{
+			const uint32 Hi = Characters[c * 2];
+			const uint32 Lo = Characters[c * 2 + 1];
+
+			for (int32 y = 0; y < CharacterRes; y++)
+			for (int32 x = 0; x < CharacterRes; x++)
+			{
+				uint8* Dest = (uint8*)(DestBuffer + x * BytesPerPixel + (c * CharacterStride) +  y * DestStride);
+
+				const uint32 C = y < CharacterRes/2 ? Hi : Lo;
+				const bool bFilled = (C & (1u << ((y * CharacterRes + x) & 31u)));
+				Dest[0] = bFilled ? 0xFF : 0x0;
+				Dest[1] = bFilled ? 0xFF : 0x0;
+				Dest[2] = bFilled ? 0xFF : 0x0;
+				Dest[3] = bFilled ? 0xFF : 0x0;
+			}
+		}
+		RHICmdList.UnlockTexture2D((FTexture2DRHIRef&)AsciiTexture->GetRenderTargetItem().ShaderResourceTexture, 0, false);
+	}
+
 	// The PreintegratedGF maybe used on forward shading inluding mobile platorm, intialize it anyway.
 	{
 		// for testing, with 128x128 R8G8 we are very close to the reference (if lower res is needed we might have to add an offset to counter the 0.5f texel shift)
diff --git a/Engine/Source/Runtime/Renderer/Private/SystemTextures.h b/Engine/Source/Runtime/Renderer/Private/SystemTextures.h
--- a/Engine/Source/Runtime/Renderer/Private/SystemTextures.h
+++ b/Engine/Source/Runtime/Renderer/Private/SystemTextures.h
@@ -144,6 +144,9 @@ public:
 	// SRV for StencilDummy Texture.
 	TRefCountPtr<FRHIShaderResourceView> StencilDummySRV;
 
+	// ASCII Standard character set - IBM code page 437 (character 32-127)
+	TRefCountPtr<IPooledRenderTarget> AsciiTexture;
+
 	// Create simple default texture
 	FRDGTextureRef GetWhiteDummy(FRDGBuilder& GraphBuilder) const;
 	FRDGTextureRef GetBlackDummy(FRDGBuilder& GraphBuilder) const;
