Unify compiler info and source info in CCHeader meta data.

#rnx

diff --git a/Engine/Source/Developer/Apple/MetalShaderFormat/Private/MetalDerivedData.cpp b/Engine/Source/Developer/Apple/MetalShaderFormat/Private/MetalDerivedData.cpp
--- a/Engine/Source/Developer/Apple/MetalShaderFormat/Private/MetalDerivedData.cpp
+++ b/Engine/Source/Developer/Apple/MetalShaderFormat/Private/MetalDerivedData.cpp
@@ -423,8 +423,6 @@ bool FMetalShaderOutputCooker::Build(TArray<uint8>& OutData)
 		
 		CrossCompiler::FHlslccHeaderWriter CCHeaderWriter;
 
-		FString MetaData = FString::Printf(TEXT("// ! %s/%s.usf:%s\n"), *Input.DebugGroupName, *Input.GetSourceFilename(), *Input.EntryPointName);
-		//FString EntryPoint = Input.EntryPointName;
 		EHlslShaderFrequency Freq = Frequency;
 		FString ALNString;
 		uint32 IABOffsetIndex = 0;
@@ -1181,8 +1179,10 @@ bool FMetalShaderOutputCooker::Build(TArray<uint8>& OutData)
 				CCHeaderWriter.WriteSideTable(TEXT("spvBufferSizeConstants"), SideTableIndex);
 			}
 
-			MetaData += TEXT("// Compiled by ShaderConductor\n");
-			MetaData += CCHeaderWriter.ToString();
+			CCHeaderWriter.WriteSourceInfo(*Input.GetSourceFilename(), *Input.EntryPointName, *Input.DebugGroupName);
+			CCHeaderWriter.WriteCompilerInfo();
+
+			FString MetaData = CCHeaderWriter.ToString();
 			MetaData += TEXT("\n\n");
 			if (ALNString.Len())
 			{
diff --git a/Engine/Source/Developer/ShaderCompilerCommon/Private/HlslccHeaderWriter.cpp b/Engine/Source/Developer/ShaderCompilerCommon/Private/HlslccHeaderWriter.cpp
--- a/Engine/Source/Developer/ShaderCompilerCommon/Private/HlslccHeaderWriter.cpp
+++ b/Engine/Source/Developer/ShaderCompilerCommon/Private/HlslccHeaderWriter.cpp
@@ -21,6 +21,26 @@ namespace CrossCompiler
 		MetaData += FString::Printf(Fmt, Args...);
 	}
 
+	void FHlslccHeaderWriter::WriteSourceInfo(const TCHAR* SourceName, const TCHAR* EntryPointName, const TCHAR* DebugGroupName)
+	{
+		check(SourceName != nullptr);
+		check(EntryPointName != nullptr);
+		if (DebugGroupName != nullptr)
+		{
+			Strings.SourceInfo = FString::Printf(TEXT("%s/%s:%s"), DebugGroupName, SourceName, EntryPointName);
+		}
+		else
+		{
+			Strings.SourceInfo = FString::Printf(TEXT("%s:%s"), SourceName, EntryPointName);
+		}
+	}
+
+	void FHlslccHeaderWriter::WriteCompilerInfo(const TCHAR* CompilerName)
+	{
+		check(CompilerName != nullptr);
+		Strings.CompilerInfo = CompilerName;
+	}
+
 	void FHlslccHeaderWriter::WriteInputAttribute(const SpvReflectInterfaceVariable& Attribute)
 	{
 		WriteIOAttribute(Strings.InputAttributes, Attribute, /*bIsInput:*/ true);
@@ -500,6 +520,16 @@ namespace CrossCompiler
 	{
 		FString MetaData;
 
+		if (!Strings.SourceInfo.IsEmpty())
+		{
+			MetaData += FString::Printf(TEXT("// ! %s\n"), *Strings.SourceInfo);
+		}
+
+		if (!Strings.CompilerInfo.IsEmpty())
+		{
+			MetaData += FString::Printf(TEXT("// Compiled by %s\n"), *Strings.CompilerInfo);
+		}
+
 		auto PrintAttributes = [&MetaData](const TCHAR* Name, const FString& Value)
 		{
 			if (!Value.IsEmpty())
diff --git a/Engine/Source/Developer/ShaderCompilerCommon/Public/HlslccHeaderWriter.h b/Engine/Source/Developer/ShaderCompilerCommon/Public/HlslccHeaderWriter.h
--- a/Engine/Source/Developer/ShaderCompilerCommon/Public/HlslccHeaderWriter.h
+++ b/Engine/Source/Developer/ShaderCompilerCommon/Public/HlslccHeaderWriter.h
@@ -17,6 +17,9 @@ namespace CrossCompiler
 	{
 	public:
 
+		void WriteSourceInfo(const TCHAR* SourceName, const TCHAR* EntryPointName, const TCHAR* DebugGroupName = nullptr);
+		void WriteCompilerInfo(const TCHAR* CompilerName = TEXT("ShaderConductor"));
+
 		void WriteInputAttribute(const SpvReflectInterfaceVariable& Attribute);
 		void WriteInputAttribute(const TCHAR* AttributeName, const TCHAR* TypeSpecifier, int32 Location, bool bLocationPrefix, bool bLocationSuffix);
 		void WriteOutputAttribute(const SpvReflectInterfaceVariable& Attribute);
@@ -46,6 +49,8 @@ namespace CrossCompiler
 	private:
 		struct FMetaDataStrings
 		{
+			FString SourceInfo;
+			FString CompilerInfo;
 			FString InputAttributes;
 			FString OutputAttributes;
 			FString UniformBlocks;
diff --git a/Engine/Source/Developer/ShaderCompilerCommon/Public/ShaderConductorContext.h b/Engine/Source/Developer/ShaderCompilerCommon/Public/ShaderConductorContext.h
--- a/Engine/Source/Developer/ShaderCompilerCommon/Public/ShaderConductorContext.h
+++ b/Engine/Source/Developer/ShaderCompilerCommon/Public/ShaderConductorContext.h
@@ -78,7 +78,7 @@ namespace CrossCompiler
 		/** Identifier for vertex input attributes: "in.var.ATTRIBUTE". */
 		const ANSICHAR* InputAttribute;
 
-		/** Identifier for globals unifrom buffers: "$Globals". */
+		/** Identifier for globals uniform buffers: "$Globals". */
 		const ANSICHAR* GlobalsUniformBuffer;
 
 		/** Identifier for the intermediate output variable in a tessellation-control shader. */
diff --git a/Engine/Source/Developer/ShaderFormatOpenGL/Private/OpenGLShaderCompiler.cpp b/Engine/Source/Developer/ShaderFormatOpenGL/Private/OpenGLShaderCompiler.cpp
--- a/Engine/Source/Developer/ShaderFormatOpenGL/Private/OpenGLShaderCompiler.cpp
+++ b/Engine/Source/Developer/ShaderFormatOpenGL/Private/OpenGLShaderCompiler.cpp
@@ -1353,8 +1353,6 @@ static bool CompileToGlslWithShaderConductor(
 
 	if (!bCompilationFailed)
 	{
-		FString MetaData;
-
 		const ANSICHAR* FrequencyPrefix = GetFrequencyPrefix(Frequency);
 
 		// Now perform reflection on the SPIRV and tweak any decorations that we need to.
@@ -2090,8 +2088,11 @@ static bool CompileToGlslWithShaderConductor(
 					TextureIndex += SamplerCount;
 				}
 
-				MetaData += TEXT("// Compiled by ShaderConductor\n");
-				MetaData += CCHeaderWriter.ToString();
+				// Generate meta data for CCHeader
+				CCHeaderWriter.WriteSourceInfo(*Input.GetSourceFilename(), *Input.EntryPointName, *Input.DebugGroupName);
+				CCHeaderWriter.WriteCompilerInfo();
+
+				const FString MetaData = CCHeaderWriter.ToString();
 
 				// Merge meta data and GLSL source to output string
 				const int32 GlslShaderSourceLen = MetaData.Len() + static_cast<int32>(GlslSource.size()) + 1;
diff --git a/Engine/Source/Developer/VulkanShaderFormat/Private/VulkanShaderCompiler.cpp b/Engine/Source/Developer/VulkanShaderFormat/Private/VulkanShaderCompiler.cpp
--- a/Engine/Source/Developer/VulkanShaderFormat/Private/VulkanShaderCompiler.cpp
+++ b/Engine/Source/Developer/VulkanShaderFormat/Private/VulkanShaderCompiler.cpp
@@ -2223,9 +2223,10 @@ static bool BuildShaderOutputFromSpirv(
 	// Build final shader output with meta data
 	FString DebugName = Input.DumpDebugInfoPath.Right(Input.DumpDebugInfoPath.Len() - Input.DumpDebugInfoRootPath.Len());
 
-	FString MetaData;// = FString::Printf(TEXT("// ! %s/%s:%s\n"), *Input.DebugGroupName, *Input.GetSourceFilename(), *Input.EntryPointName);
-	MetaData += TEXT("// Compiled by ShaderConductor\n");
-	MetaData += CCHeaderWriter.ToString();
+	CCHeaderWriter.WriteSourceInfo(*Input.GetSourceFilename(), *Input.EntryPointName, *Input.DebugGroupName);
+	CCHeaderWriter.WriteCompilerInfo();
+
+	const FString MetaData = CCHeaderWriter.ToString();
 
 	Output.Target = Input.Target;
 
