Fix non-unity

diff --git a/Engine/Source/Developer/ShaderCompilerCommon/Public/ShaderSymbolExport.h b/Engine/Source/Developer/ShaderCompilerCommon/Public/ShaderSymbolExport.h
--- a/Engine/Source/Developer/ShaderCompilerCommon/Public/ShaderSymbolExport.h
+++ b/Engine/Source/Developer/ShaderCompilerCommon/Public/ShaderSymbolExport.h
@@ -5,6 +5,7 @@
 #if WITH_ENGINE
 
 #include "Serialization/MemoryReader.h"
+#include "HAL/FileManager.h"
 #include "HAL/IConsoleManager.h"
 #include "HAL/PlatformFileManager.h"
 #include "GenericPlatform/GenericPlatformFile.h"
