Allows SHADER_USE_PARAMETER_STRUCT() to submit all shader parameters in a single RHISetShaderParameter()

diff --git a/Engine/Source/Developer/ShaderCompilerCommon/Private/ShaderCompilerCommon.cpp b/Engine/Source/Developer/ShaderCompilerCommon/Private/ShaderCompilerCommon.cpp
--- a/Engine/Source/Developer/ShaderCompilerCommon/Private/ShaderCompilerCommon.cpp
+++ b/Engine/Source/Developer/ShaderCompilerCommon/Private/ShaderCompilerCommon.cpp
@@ -469,6 +469,69 @@ TCHAR* FindNextUniformBufferReference(TCHAR* SearchPtr, const TCHAR* SearchStrin
 	return nullptr;
 }
 
+void AddNoteToDisplayShaderParameterStructureOnCppSide(
+	const FShaderParametersMetadata* ParametersStructure,
+	FShaderCompilerOutput& CompilerOutput)
+{
+	FShaderCompilerError Error;
+	Error.StrippedErrorMessage = FString::Printf(
+		TEXT("Note: Definition of structure %s"),
+		ParametersStructure->GetStructTypeName());
+	Error.ErrorVirtualFilePath = ANSI_TO_TCHAR(ParametersStructure->GetFileName());
+	Error.ErrorLineString = FString::FromInt(ParametersStructure->GetFileLine());
+
+	CompilerOutput.Errors.Add(Error);
+}
+
+void AddNoteToDisplayShaderParameterMemberOnCppSide(
+	const FShaderCompilerInput& CompilerInput,
+	const FShaderParameterParser::FParsedShaderParameter& ParsedParameter,
+	FShaderCompilerOutput& CompilerOutput)
+{
+	const FShaderParametersMetadata* MemberContainingStruct = nullptr;
+	const FShaderParametersMetadata::FMember* Member = nullptr;
+	{
+		int32 ArrayElementId = 0;
+		FString NamePrefix;
+		CompilerInput.RootParametersStructure->FindMemberFromOffset(ParsedParameter.ConstantBufferOffset, &MemberContainingStruct, &Member, &ArrayElementId, &NamePrefix);
+	}
+
+	FString CppCodeName = CompilerInput.RootParametersStructure->GetFullMemberCodeName(ParsedParameter.ConstantBufferOffset);
+
+	FShaderCompilerError Error;
+	Error.StrippedErrorMessage = FString::Printf(
+		TEXT("Note: Definition of %s"),
+		*CppCodeName);
+	Error.ErrorVirtualFilePath = ANSI_TO_TCHAR(MemberContainingStruct->GetFileName());
+	Error.ErrorLineString = FString::FromInt(Member->GetFileLine());
+
+	CompilerOutput.Errors.Add(Error);
+}
+
+void AddUnboundShaderParameterError(
+	const FShaderCompilerInput& CompilerInput,
+	const FShaderParameterParser& ShaderParameterParser,
+	const FString& ParameterBindingName,
+	FShaderCompilerOutput& CompilerOutput)
+{
+	check(CompilerInput.RootParametersStructure);
+
+	const FShaderParameterParser::FParsedShaderParameter& Member = ShaderParameterParser.FindParameterInfos(ParameterBindingName);
+	check(!Member.IsBindable());
+
+	FShaderCompilerError Error(FString::Printf(
+		TEXT("Error: Shader parameter %s could not be bound to %s's shader parameter structure %s."),
+		*ParameterBindingName,
+		*CompilerInput.ShaderName,
+		CompilerInput.RootParametersStructure->GetStructTypeName()));
+	ShaderParameterParser.GetParameterFileAndLine(Member, Error.ErrorVirtualFilePath, Error.ErrorLineString);
+
+	CompilerOutput.Errors.Add(Error);
+	CompilerOutput.bSucceeded = false;
+
+	AddNoteToDisplayShaderParameterStructureOnCppSide(CompilerInput.RootParametersStructure, CompilerOutput);
+}
+
 bool FShaderParameterParser::ParseAndMoveShaderParametersToRootConstantBuffer(
 	const FShaderCompilerInput& CompilerInput,
 	FShaderCompilerOutput& CompilerOutput,
@@ -484,21 +547,21 @@ bool FShaderParameterParser::ParseAndMoveShaderParametersToRootConstantBuffer(
 	const bool bMoveToRootConstantBuffer = ConstantBufferType != nullptr;
 	OriginalParsedShader = PreprocessedShaderSource;
 
+	// Reserves the number of parameters up front.
+	ParsedParameters.Reserve(CompilerInput.RootParametersStructure->GetSize() / sizeof(int32));
+
 	CompilerInput.RootParametersStructure->IterateShaderParameterMembers(
 		[&](const FShaderParametersMetadata& ParametersMetadata,
 			const FShaderParametersMetadata::FMember& Member,
 			const TCHAR* ShaderBindingName,
 			uint16 ByteOffset)
 	{
-		if (
-			Member.GetBaseType() != UBMT_INT32 &&
-			Member.GetBaseType() != UBMT_UINT32 &&
-			Member.GetBaseType() != UBMT_FLOAT32)
-		{
-			return;
-		}
+		FParsedShaderParameter ParsedParameter;
+		ParsedParameter.Member = &Member;
+		ParsedParameter.ConstantBufferOffset = ByteOffset;
+		check(ParsedParameter.IsBindable());
 
-		ParsedParameters.Add(ShaderBindingName, FParsedShaderParameter());
+		ParsedParameters.Add(ShaderBindingName, ParsedParameter);
 	});
 
 	bool bSuccess = true;
@@ -538,6 +601,8 @@ bool FShaderParameterParser::ParseAndMoveShaderParametersToRootConstantBuffer(
 		int32 TypeEndPos = -1;
 		int32 NameStartPos = -1;
 		int32 NameEndPos = -1;
+		int32 ArrayStartPos = -1;
+		int32 ArrayEndPos = -1;
 		int32 ScopeIndent = 0;
 
 		EState State = EState::Scanning;
@@ -549,6 +614,8 @@ bool FShaderParameterParser::ParseAndMoveShaderParametersToRootConstantBuffer(
 			TypeEndPos = -1;
 			NameStartPos = -1;
 			NameEndPos = -1;
+			ArrayStartPos = -1;
+			ArrayEndPos = -1;
 			State = EState::Scanning;
 		};
 
@@ -582,6 +649,9 @@ bool FShaderParameterParser::ParseAndMoveShaderParametersToRootConstantBuffer(
 				FString Type = PreprocessedShaderSource.Mid(TypeStartPos, TypeEndPos - TypeStartPos + 1);
 				FString Name = PreprocessedShaderSource.Mid(NameStartPos, NameEndPos - NameStartPos + 1);
 
+				FParsedShaderParameter ParsedParameter;
+				bool bUpdateParsedParameters = false;
+				bool bEraseOriginalParameter = false;
 				if (ParsedParameters.Contains(Name))
 				{
 					if (ParsedParameters.FindChecked(Name).IsFound())
@@ -590,23 +660,48 @@ bool FShaderParameterParser::ParseAndMoveShaderParametersToRootConstantBuffer(
 					}
 					else
 					{
-						FParsedShaderParameter ParsedParameter;
-						ParsedParameter.Type = Type;
-						ParsedParameter.PragamLineoffset = CurrentPragamLineoffset;
-						ParsedParameter.LineOffset = CurrentLineoffset;
-						ParsedParameters[Name] = ParsedParameter;
-
-						// Erases this shader parameter conserving the same line numbers.
-						if (bMoveToRootConstantBuffer)
+						// Update the parsed parameters
+						bUpdateParsedParameters = true;
+						ParsedParameter = ParsedParameters.FindChecked(Name);
+
+						// Erase the parameter to move it into the root constant buffer.
+						if (bMoveToRootConstantBuffer && ParsedParameter.IsBindable())
 						{
-							for (int32 j = TypeStartPos; j <= Cursor; j++)
-							{
-								if (PreprocessedShaderSource[j] != '\r' && PreprocessedShaderSource[j] != '\n')
-									PreprocessedShaderSource[j] = ' ';
-							}
+							EUniformBufferBaseType BaseType = ParsedParameter.Member->GetBaseType();
+							bEraseOriginalParameter = BaseType == UBMT_INT32 || BaseType == UBMT_UINT32 || BaseType == UBMT_FLOAT32;
 						}
 					}
 				}
+				else
+				{
+					// Update the parsed parameters to still have file and line number.
+					bUpdateParsedParameters = true;
+				}
+
+				// Update 
+				if (bUpdateParsedParameters)
+				{
+					ParsedParameter.ParsedType = Type;
+					ParsedParameter.ParsedPragmaLineoffset = CurrentPragamLineoffset;
+					ParsedParameter.ParsedLineOffset = CurrentLineoffset;
+
+					if (ArrayStartPos != -1 && ArrayEndPos != -1)
+					{
+						ParsedParameter.ParsedArraySize = PreprocessedShaderSource.Mid(ArrayStartPos + 1, ArrayEndPos - ArrayStartPos - 1);
+					}
+
+					ParsedParameters.Add(Name, ParsedParameter);
+				}
+
+				// Erases this shader parameter conserving the same line numbers.
+				if (bEraseOriginalParameter)
+				{
+					for (int32 j = TypeStartPos; j <= Cursor; j++)
+					{
+						if (PreprocessedShaderSource[j] != '\r' && PreprocessedShaderSource[j] != '\n')
+							PreprocessedShaderSource[j] = ' ';
+					}
+				}
 
 				ResetState();
 			};
@@ -798,6 +893,7 @@ bool FShaderParameterParser::ParseAndMoveShaderParametersToRootConstantBuffer(
 					// Syntax:
 					//  uint MyArray[
 					NameEndPos = Cursor - 1;
+					ArrayStartPos = Cursor;
 					State = EState::ParsingPotentialArraySize;
 				}
 				else if (bIsWhiteSpace)
@@ -844,6 +940,7 @@ bool FShaderParameterParser::ParseAndMoveShaderParametersToRootConstantBuffer(
 					{
 						// Syntax:
 						//  uint MyArray [
+						ArrayStartPos = Cursor;
 						State = EState::ParsingPotentialArraySize;
 					}
 					else
@@ -865,6 +962,7 @@ bool FShaderParameterParser::ParseAndMoveShaderParametersToRootConstantBuffer(
 			{
 				if (Char == ']')
 				{
+					ArrayEndPos = Cursor;
 					State = EState::FinishedArraySize;
 				}
 				else if (Char == ';')
@@ -913,7 +1011,7 @@ bool FShaderParameterParser::ParseAndMoveShaderParametersToRootConstantBuffer(
 				return;
 			}
 
-			const FParsedShaderParameter& ParsedParameter = ParsedParameters[ShaderBindingName];
+			FParsedShaderParameter& ParsedParameter = ParsedParameters[ShaderBindingName];
 
 			if (ParsedParameter.IsFound())
 			{
@@ -934,26 +1032,43 @@ bool FShaderParameterParser::ParseAndMoveShaderParametersToRootConstantBuffer(
 					case 12:
 						HLSLOffset.Append(TEXT(".w"));
 						break;
+					default:
+						unimplemented();
 					}
 				}
 
-				RootCBufferContent.Append(FString::Printf(
-					TEXT("%s %s : packoffset(c%s);\r\n"),
-					*ParsedParameter.Type,
-					ShaderBindingName,
-					*HLSLOffset));
+				if (!ParsedParameter.ParsedArraySize.IsEmpty())
+				{
+					RootCBufferContent.Append(FString::Printf(
+						TEXT("%s %s[%s] : packoffset(c%s);\r\n"),
+						*ParsedParameter.ParsedType,
+						ShaderBindingName,
+						*ParsedParameter.ParsedArraySize,
+						*HLSLOffset));
+				}
+				else
+				{
+					RootCBufferContent.Append(FString::Printf(
+						TEXT("%s %s : packoffset(c%s);\r\n"),
+						*ParsedParameter.ParsedType,
+						ShaderBindingName,
+						*HLSLOffset));
+				}
 			}
 		});
 
-		FString NewShaderCode = FString::Printf(
+		FString CBufferCodeBlock = FString::Printf(
 			TEXT("%s %s\r\n")
 			TEXT("{\r\n")
 			TEXT("%s")
-			TEXT("}\r\n\r\n%s"),
+			TEXT("}\r\n\r\n"),
 			ConstantBufferType,
 			FShaderParametersMetadata::kRootUniformBufferBindingName,
-			*RootCBufferContent,
-			*PreprocessedShaderSource);
+			*RootCBufferContent);
+
+		FString NewShaderCode = (
+			MakeInjectedShaderCodeBlock(TEXT("ParseAndMoveShaderParametersToRootConstantBuffer"), CBufferCodeBlock) +
+			PreprocessedShaderSource);
 
 		PreprocessedShaderSource = MoveTemp(NewShaderCode);
 	}
@@ -961,74 +1076,40 @@ bool FShaderParameterParser::ParseAndMoveShaderParametersToRootConstantBuffer(
 	return bSuccess;
 }
 
-void FShaderParameterParser::ValidateShaderParameterTypes(
+void FShaderParameterParser::ValidateShaderParameterType(
 	const FShaderCompilerInput& CompilerInput,
+	const TCHAR* ShaderBindingName,
+	const FParsedShaderParameter& ParsedParameter,
+	int32 BoundSize,
 	FShaderCompilerOutput& CompilerOutput) const
 {
-	// The shader doesn't have any parameter binding through shader structure, therefore don't do anything.
-	if (!CompilerInput.RootParametersStructure)
-	{
-		return;
-	}
-
-	if (!CompilerOutput.bSucceeded)
-	{
-		return;
-	}
+	check(ParsedParameter.IsFound());
+	check(CompilerInput.RootParametersStructure);
 
 	EShaderPlatform ShaderPlatform = EShaderPlatform(CompilerInput.Target.Platform);
 
-	const TMap<FString, FParameterAllocation>& ParametersFoundByCompiler = CompilerOutput.ParameterMap.GetParameterMap();
-
-	bool bSuccess = true;
-	CompilerInput.RootParametersStructure->IterateShaderParameterMembers(
-		[&](const FShaderParametersMetadata& ParametersMetadata,
-			const FShaderParametersMetadata::FMember& Member,
-			const TCHAR* ShaderBindingName,
-			uint16 ByteOffset)
+	// Validate the shader type.
 	{
-		if (
-			Member.GetBaseType() != UBMT_INT32 &&
-			Member.GetBaseType() != UBMT_UINT32 &&
-			Member.GetBaseType() != UBMT_FLOAT32)
-		{
-			return;
-		}
-
-		const FParsedShaderParameter& ParsedParameter = ParsedParameters[ShaderBindingName];
-
-		// Did not find shader parameter in code.
-		if (!ParsedParameter.IsFound())
-		{
-			// Verify the shader compiler also did not find this parameter to make sure there is no bug in the parser.
-			checkf(
-				!ParametersFoundByCompiler.Contains(ShaderBindingName),
-				TEXT("Looks like there is a bug in FShaderParameterParser ParameterName=%s DumpDebugInfoPath=%s"),
-				ShaderBindingName,
-				*CompilerInput.DumpDebugInfoPath);
-			return;
-		}
-
 		FString ExpectedShaderType;
-		Member.GenerateShaderParameterType(ExpectedShaderType, ShaderPlatform);
+		ParsedParameter.Member->GenerateShaderParameterType(ExpectedShaderType, ShaderPlatform);
 
-		const bool bShouldBeInt = Member.GetBaseType() == UBMT_INT32;
-		const bool bShouldBeUint = Member.GetBaseType() == UBMT_UINT32;
+		const bool bShouldBeInt = ParsedParameter.Member->GetBaseType() == UBMT_INT32;
+		const bool bShouldBeUint = ParsedParameter.Member->GetBaseType() == UBMT_UINT32;
 
 		// Match parsed type with expected shader type
-		bool bIsTypeCorrect = ParsedParameter.Type == ExpectedShaderType;
-		
+		bool bIsTypeCorrect = ParsedParameter.ParsedType == ExpectedShaderType;
+
 		if (!bIsTypeCorrect)
 		{
 			// Accept half-precision floats when single-precision was requested
-			if (ParsedParameter.Type.StartsWith(TEXT("half")) && Member.GetBaseType() == UBMT_FLOAT32)
+			if (ParsedParameter.ParsedType.StartsWith(TEXT("half")) && ParsedParameter.Member->GetBaseType() == UBMT_FLOAT32)
 			{
-				bIsTypeCorrect = (FCString::Strcmp(*ParsedParameter.Type + 4, *ExpectedShaderType + 5) == 0);
+				bIsTypeCorrect = (FCString::Strcmp(*ParsedParameter.ParsedType + 4, *ExpectedShaderType + 5) == 0);
 			}
 			// Accept single-precision floats when half-precision was expected
-			else if (ParsedParameter.Type.StartsWith(TEXT("float")) && ExpectedShaderType.StartsWith(TEXT("half")))
+			else if (ParsedParameter.ParsedType.StartsWith(TEXT("float")) && ExpectedShaderType.StartsWith(TEXT("half")))
 			{
-				bIsTypeCorrect = (FCString::Strcmp(*ParsedParameter.Type + 5, *ExpectedShaderType + 4) == 0);
+				bIsTypeCorrect = (FCString::Strcmp(*ParsedParameter.ParsedType + 5, *ExpectedShaderType + 4) == 0);
 			}
 		}
 
@@ -1048,25 +1129,103 @@ void FShaderParameterParser::ValidateShaderParameterTypes(
 				NewExpectedShaderType.RemoveAt(0);
 			}
 
-			bIsTypeCorrect = ParsedParameter.Type == NewExpectedShaderType;
+			bIsTypeCorrect = ParsedParameter.ParsedType == NewExpectedShaderType;
 		}
 
 		if (!bIsTypeCorrect)
 		{
+			FString CppCodeName = CompilerInput.RootParametersStructure->GetFullMemberCodeName(ParsedParameter.ConstantBufferOffset);
+
 			FShaderCompilerError Error;
 			Error.StrippedErrorMessage = FString::Printf(
-				TEXT("Type %s of shader parameter %s in shader mismatch the shader parameter structure: it expects a %s"),
-				*ParsedParameter.Type,
+				TEXT("Error: Type %s of shader parameter %s in shader mismatch the shader parameter structure: %s expects a %s"),
+				*ParsedParameter.ParsedType,
 				ShaderBindingName,
+				*CppCodeName,
 				*ExpectedShaderType);
-			ExtractFileAndLine(ParsedParameter.PragamLineoffset, ParsedParameter.LineOffset, Error.ErrorVirtualFilePath, Error.ErrorLineString);
+			GetParameterFileAndLine(ParsedParameter, Error.ErrorVirtualFilePath, Error.ErrorLineString);
 
 			CompilerOutput.Errors.Add(Error);
-			bSuccess = false;
+			CompilerOutput.bSucceeded = false;
+
+			AddNoteToDisplayShaderParameterMemberOnCppSide(CompilerInput, ParsedParameter, CompilerOutput);
 		}
-	});
+	}
+
+	// Validate parameter size, in case this is an array.
+	if (BoundSize > int32(ParsedParameter.Member->GetMemberSize()))
+	{
+		FString CppCodeName = CompilerInput.RootParametersStructure->GetFullMemberCodeName(ParsedParameter.ConstantBufferOffset);
+
+		FShaderCompilerError Error;
+		Error.StrippedErrorMessage = FString::Printf(
+			TEXT("Error: The size required to bind shader parameter %s is %i bytes, smaller than %s that is %i bytes in the parameter structure."),
+			ShaderBindingName,
+			BoundSize,
+			*CppCodeName,
+			ParsedParameter.Member->GetMemberSize());
+		GetParameterFileAndLine(ParsedParameter, Error.ErrorVirtualFilePath, Error.ErrorLineString);
+
+		CompilerOutput.Errors.Add(Error);
+		CompilerOutput.bSucceeded = false;
+
+		AddNoteToDisplayShaderParameterMemberOnCppSide(CompilerInput, ParsedParameter, CompilerOutput);
+	}
+}
+
+void FShaderParameterParser::ValidateShaderParameterTypes(
+	const FShaderCompilerInput& CompilerInput,
+	FShaderCompilerOutput& CompilerOutput) const
+{
+	// The shader doesn't have any parameter binding through shader structure, therefore don't do anything.
+	if (!CompilerInput.RootParametersStructure)
+	{
+		return;
+	}
 
-	CompilerOutput.bSucceeded = bSuccess;
+	if (!CompilerOutput.bSucceeded)
+	{
+		return;
+	}
+
+	const TMap<FString, FParameterAllocation>& ParametersFoundByCompiler = CompilerOutput.ParameterMap.GetParameterMap();
+
+	CompilerInput.RootParametersStructure->IterateShaderParameterMembers(
+		[&](const FShaderParametersMetadata& ParametersMetadata,
+			const FShaderParametersMetadata::FMember& Member,
+			const TCHAR* ShaderBindingName,
+			uint16 ByteOffset)
+	{
+		if (
+			Member.GetBaseType() != UBMT_INT32 &&
+			Member.GetBaseType() != UBMT_UINT32 &&
+			Member.GetBaseType() != UBMT_FLOAT32)
+		{
+			return;
+		}
+
+		const FParsedShaderParameter& ParsedParameter = ParsedParameters[ShaderBindingName];
+
+		// Did not find shader parameter in code.
+		if (!ParsedParameter.IsFound())
+		{
+			// Verify the shader compiler also did not find this parameter to make sure there is no bug in the parser.
+			checkf(
+				!ParametersFoundByCompiler.Contains(ShaderBindingName),
+				TEXT("Looks like there is a bug in FShaderParameterParser ParameterName=%s DumpDebugInfoPath=%s"),
+				ShaderBindingName,
+				*CompilerInput.DumpDebugInfoPath);
+			return;
+		}
+
+		int32 BoundSize = 0;
+		if (const FParameterAllocation* ParameterAllocation = ParametersFoundByCompiler.Find(ShaderBindingName))
+		{
+			BoundSize = ParameterAllocation->Size;
+		}
+
+		ValidateShaderParameterType(CompilerInput, ShaderBindingName, ParsedParameter, BoundSize, CompilerOutput);
+	});
 }
 
 void FShaderParameterParser::ExtractFileAndLine(int32 PragamLineoffset, int32 LineOffset, FString& OutFile, FString& OutLine) const
diff --git a/Engine/Source/Developer/ShaderCompilerCommon/Public/ShaderCompilerCommon.h b/Engine/Source/Developer/ShaderCompilerCommon/Public/ShaderCompilerCommon.h
--- a/Engine/Source/Developer/ShaderCompilerCommon/Public/ShaderCompilerCommon.h
+++ b/Engine/Source/Developer/ShaderCompilerCommon/Public/ShaderCompilerCommon.h
@@ -43,34 +43,75 @@ extern SHADERCOMPILERCOMMON_API void BuildResourceTableTokenStream(
 // Finds the number of used uniform buffers in a resource map
 extern SHADERCOMPILERCOMMON_API int16 GetNumUniformBuffersUsed(const FShaderCompilerResourceTable& InSRT);
 
-
 /** Validates and moves all the shader loose data parameter defined in the root scope of the shader into the root uniform buffer. */
 class SHADERCOMPILERCOMMON_API FShaderParameterParser
 {
 public:
+	struct FParsedShaderParameter
+	{
+	public:
+		/** Original information about the member. */
+		const FShaderParametersMetadata::FMember* Member = nullptr;
+
+		/** Information found about the member when parsing the preprocessed code. */
+		FString ParsedType;
+		FString ParsedArraySize;
+
+		/** Offset the member should be in the constant buffer. */
+		int32 ConstantBufferOffset = 0;
+
+		/* Returns whether the shader parameter has been found when parsing. */
+		bool IsFound() const
+		{
+			return !ParsedType.IsEmpty();
+		}
+
+		/** Returns whether the shader parameter is bindable to the shader parameter structure. */
+		bool IsBindable() const
+		{
+			return Member != nullptr;
+		}
+
+	private:
+		int32 ParsedPragmaLineoffset = 0;
+		int32 ParsedLineOffset = 0;
+
+		friend class FShaderParameterParser;
+	};
+
+	/** Parses the preprocessed shader code and move the parameters into root constant buffer */
 	bool ParseAndMoveShaderParametersToRootConstantBuffer(
 		const FShaderCompilerInput& CompilerInput,
 		FShaderCompilerOutput& CompilerOutput,
 		FString& PreprocessedShaderSource,
 		const TCHAR* ConstantBufferType);
 
+	/** Gets parsing information from a parameter binding name. */
+	const FParsedShaderParameter& FindParameterInfos(const FString& ParameterName) const
+	{
+		return ParsedParameters.FindChecked(ParameterName);
+	}
+
+	/** Validates the shader parameter in code is compatible with the shader parameter structure. */
+	void ValidateShaderParameterType(
+		const FShaderCompilerInput& CompilerInput,
+		const TCHAR* ShaderBindingName,
+		const FParsedShaderParameter& ParsedParameter,
+		int32 BoundSize,
+		FShaderCompilerOutput& CompilerOutput) const;
+
+	/** Validates shader parameter map is compatible with the shader parameter structure. */
 	void ValidateShaderParameterTypes(
 		const FShaderCompilerInput& CompilerInput,
 		FShaderCompilerOutput& CompilerOutput) const;
 
-private:
-	struct FParsedShaderParameter
+	/** Gets file and line of the parameter in the shader source code. */
+	void GetParameterFileAndLine(const FParsedShaderParameter& ParsedParameter, FString& OutFile, FString& OutLine) const
 	{
-		FString Type;
-		int32 PragamLineoffset;
-		int32 LineOffset;
-
-		bool IsFound() const
-		{
-			return !Type.IsEmpty();
-		}
-	};
+		return ExtractFileAndLine(ParsedParameter.ParsedPragmaLineoffset, ParsedParameter.ParsedLineOffset, OutFile, OutLine);
+	}
 
+private:
 	void ExtractFileAndLine(int32 PragamLineoffset, int32 LineOffset, FString& OutFile, FString& OutLine) const;
 
 	FString OriginalParsedShader;
@@ -78,6 +119,24 @@ private:
 	TMap<FString, FParsedShaderParameter> ParsedParameters;
 };
 
+/** Adds a note to CompilerOutput.Error about where the shader parameter structure is on C++ side. */
+extern SHADERCOMPILERCOMMON_API void AddNoteToDisplayShaderParameterStructureOnCppSide(
+	const FShaderParametersMetadata* ParametersStructure,
+	FShaderCompilerOutput& CompilerOutput);
+
+/** Adds a note to CompilerOutput.Error about where the shader parameter is on C++ side. */
+extern SHADERCOMPILERCOMMON_API void AddNoteToDisplayShaderParameterMemberOnCppSide(
+	const FShaderCompilerInput& CompilerInput,
+	const FShaderParameterParser::FParsedShaderParameter& ParsedParameter,
+	FShaderCompilerOutput& CompilerOutput);
+
+/** Adds an error to CompilerOutput.Error about a shader parameters that could not be bound. */
+extern SHADERCOMPILERCOMMON_API void AddUnboundShaderParameterError(
+	const FShaderCompilerInput& CompilerInput,
+	const FShaderParameterParser& ShaderParameterParser,
+	const FString& ParameterBindingName,
+	FShaderCompilerOutput& CompilerOutput);
+
 // The cross compiler doesn't yet support struct initializers needed to construct static structs for uniform buffers
 // Replace all uniform buffer struct member references (View.WorldToClip) with a flattened name that removes the struct dependency (View_WorldToClip)
 extern SHADERCOMPILERCOMMON_API void RemoveUniformBuffersFromSource(const FShaderCompilerEnvironment& Environment, FString& PreprocessedShaderSource);
diff --git a/Engine/Source/Developer/ShaderPreprocessor/Private/ShaderPreprocessor.cpp b/Engine/Source/Developer/ShaderPreprocessor/Private/ShaderPreprocessor.cpp
--- a/Engine/Source/Developer/ShaderPreprocessor/Private/ShaderPreprocessor.cpp
+++ b/Engine/Source/Developer/ShaderPreprocessor/Private/ShaderPreprocessor.cpp
@@ -200,10 +200,13 @@ static void DumpShaderDefinesAsCommentedCode(const FShaderCompilerInput& ShaderI
 	Definitions.GetKeys(/* out */ Keys);
 	Keys.Sort();
 
+	FString Defines;
 	for (const FString& Key : Keys)
 	{
-		*OutDefines += FString::Printf(TEXT("// #define %s %s\n"), *Key, *Definitions[Key]);
+		Defines += FString::Printf(TEXT("// #define %s %s\n"), *Key, *Definitions[Key]);
 	}
+
+	*OutDefines = MakeInjectedShaderCodeBlock(TEXT("DumpShaderDefinesAsCommentedCode"), Defines);
 }
 
 //////////////////////////////////////////////////////////////////////////
diff --git a/Engine/Source/Programs/ShaderCompileWorker/Private/ShaderCompileWorker.cpp b/Engine/Source/Programs/ShaderCompileWorker/Private/ShaderCompileWorker.cpp
--- a/Engine/Source/Programs/ShaderCompileWorker/Private/ShaderCompileWorker.cpp
+++ b/Engine/Source/Programs/ShaderCompileWorker/Private/ShaderCompileWorker.cpp
@@ -396,6 +396,32 @@ private:
 			return CharName;
 		};
 
+		// Array of string used as const ANSICHAR* during compilation process.
+		TArray<TUniquePtr<TArray<ANSICHAR>>> AllocatedAnsiStrings;
+		auto DeserializeConstANSICHAR = [&AllocatedAnsiStrings](FArchive& Archive)
+		{
+			FString Name;
+			Archive << Name;
+
+			const ANSICHAR* CharName = nullptr;
+			if (Name.Len() != 0)
+			{
+				if (AllocatedAnsiStrings.GetSlack() == 0)
+				{
+					AllocatedAnsiStrings.Reserve(AllocatedAnsiStrings.Num() + 1024);
+				}
+
+				TArray<ANSICHAR> AnsiString;
+				AnsiString.SetNumZeroed(Name.Len() + 1);
+				ANSICHAR* Dest = &AnsiString[0];
+				FCStringAnsi::Strcpy(Dest, Name.Len() + 1, TCHAR_TO_ANSI(*Name));
+
+				AllocatedAnsiStrings.Add(MakeUnique<TArray<ANSICHAR>>(AnsiString));
+				CharName = &(*AllocatedAnsiStrings.Last())[0];
+			}
+			return CharName;
+		};
+
 		// Shared inputs
 		TMap<FString, FThreadSafeSharedStringPtr> ExternalIncludes;
 		{
@@ -441,6 +467,7 @@ private:
 				const TCHAR* StructTypeName;
 				const TCHAR* ShaderVariableName;
 				FShaderParametersMetadata::EUseCase UseCase;
+				const ANSICHAR* StructFileName;
 				int32 StructFileLine;
 				uint32 Size;
 				int32 MemberCount;
@@ -449,6 +476,7 @@ private:
 				StructTypeName = DeserializeConstTCHAR(InputFile);
 				ShaderVariableName = DeserializeConstTCHAR(InputFile);
 				InputFile << UseCase;
+				StructFileName = DeserializeConstANSICHAR(InputFile);
 				InputFile << StructFileLine;
 				InputFile << Size;
 				InputFile << MemberCount;
@@ -515,7 +543,7 @@ private:
 					/* InStructTypeName = */ StructTypeName,
 					/* InShaderVariableName = */ ShaderVariableName,
 					/* InStaticSlotName = */ nullptr,
-					/* InFileName = */ nullptr,
+					StructFileName,
 					StructFileLine,
 					Size,
 					Members,
diff --git a/Engine/Source/Runtime/Engine/Private/ShaderCompiler/ShaderCompiler.cpp b/Engine/Source/Runtime/Engine/Private/ShaderCompiler/ShaderCompiler.cpp
--- a/Engine/Source/Runtime/Engine/Private/ShaderCompiler/ShaderCompiler.cpp
+++ b/Engine/Source/Runtime/Engine/Private/ShaderCompiler/ShaderCompiler.cpp
@@ -1223,6 +1223,7 @@ bool FShaderCompileUtilities::DoWriteTasks(const TArray<FShaderCommonCompileJobP
 			FString StructTypeName = Struct->GetStructTypeName();
 			FString ShaderVariableName = Struct->GetShaderVariableName();
 			uint8 UseCase = uint8(Struct->GetUseCase());
+			FString StructFileName = FString(ANSI_TO_TCHAR(Struct->GetFileName()));
 			int32 StructFileLine = Struct->GetFileLine();
 			uint32 Size = Struct->GetSize();
 			int32 MemberCount = Struct->GetMembers().Num();
@@ -1233,6 +1234,7 @@ bool FShaderCompileUtilities::DoWriteTasks(const TArray<FShaderCommonCompileJobP
 			TransferFile << StructTypeName;
 			TransferFile << ShaderVariableName;
 			TransferFile << UseCase;
+			TransferFile << StructFileName;
 			TransferFile << StructFileLine;
 			TransferFile << Size;
 			TransferFile << MemberCount;
@@ -4955,6 +4957,7 @@ void GlobalBeginCompileShader(
 	Input.DebugDescription = DebugDescription;
 	Input.DebugExtension = DebugExtension;
 	Input.RootParametersStructure = ShaderType->GetRootParametersMetadata();
+	Input.ShaderName = ShaderType->GetName();
 
 	// Verify FShaderCompilerInput's file paths are consistent. 
 	#if DO_CHECK
diff --git a/Engine/Source/Runtime/RenderCore/Private/ShaderCore.cpp b/Engine/Source/Runtime/RenderCore/Private/ShaderCore.cpp
--- a/Engine/Source/Runtime/RenderCore/Private/ShaderCore.cpp
+++ b/Engine/Source/Runtime/RenderCore/Private/ShaderCore.cpp
@@ -1446,6 +1446,12 @@ void SerializeUniformBufferInfo(FShaderSaveArchive& Ar, const TMap<const TCHAR*,
 	}
 }
 
+FString MakeInjectedShaderCodeBlock(const TCHAR* BlockName, const FString& CodeToInject)
+{
+	return FString::Printf(TEXT("#line 1 \"%s\"\n%s"), BlockName, *CodeToInject);
+}
+
+
 FString FShaderCompilerError::GetShaderSourceFilePath() const
 {
 	if (IFileManager::Get().FileExists(*ErrorVirtualFilePath))
@@ -1566,6 +1572,7 @@ FArchive& operator<<(FArchive& Ar, FShaderCompilerInput& Input)
 	Ar << Input.SourceFilePrefix;
 	Ar << Input.VirtualSourceFilePath;
 	Ar << Input.EntryPointName;
+	Ar << Input.ShaderName;
 	Ar << Input.bSkipPreprocessedCache;
 	Ar << Input.bCompilingForShaderPipeline;
 	Ar << Input.bGenerateDirectCompileFile;
diff --git a/Engine/Source/Runtime/RenderCore/Private/ShaderParameterStruct.cpp b/Engine/Source/Runtime/RenderCore/Private/ShaderParameterStruct.cpp
--- a/Engine/Source/Runtime/RenderCore/Private/ShaderParameterStruct.cpp
+++ b/Engine/Source/Runtime/RenderCore/Private/ShaderParameterStruct.cpp
@@ -253,7 +253,33 @@ void FShaderParameterBindings::BindForLegacyShaderParameters(const FShader* Shad
 	BindingContext.PermutationId = PermutationId;
 	BindingContext.Bindings = this;
 	BindingContext.ParametersMap = &ParametersMap;
-	BindingContext.bUseRootShaderParameters = false;
+
+	// When using a stable root constant buffer, can submit all shader parameters in one RHISetShaderParameter()
+	{
+		const TCHAR* ShaderBindingName = FShaderParametersMetadata::kRootUniformBufferBindingName;
+		BindingContext.bUseRootShaderParameters = ParametersMap.ContainsParameterAllocation(ShaderBindingName);
+
+		if (BindingContext.bUseRootShaderParameters)
+		{
+			const FParameterAllocation& Allocation = ParametersMap.ParameterMap.FindChecked(ShaderBindingName);
+
+			check(int32(Allocation.BufferIndex) == FShaderParametersMetadata::kRootCBufferBindingIndex);
+			check(uint32(Allocation.BaseIndex + Allocation.Size) <= StructMetaData.GetSize());
+			check(Allocation.Type == EShaderParameterType::LooseData);
+
+			FShaderParameterBindings::FParameter Parameter;
+			Parameter.BufferIndex = Allocation.BufferIndex;
+			Parameter.BaseIndex = Allocation.BaseIndex;
+			Parameter.ByteOffset = Allocation.BaseIndex;
+			Parameter.ByteSize = Allocation.Size;
+
+			Parameters.Add(Parameter);
+
+			BindingContext.ShaderGlobalScopeBindings.Add(ShaderBindingName, StructMetaData.GetStructTypeName());
+			Allocation.bBound = true;
+		}
+	}
+
 	BindingContext.Bind(
 		StructMetaData,
 		/* MemberPrefix = */ TEXT(""),
diff --git a/Engine/Source/Runtime/RenderCore/Public/ShaderCompilerCore.h b/Engine/Source/Runtime/RenderCore/Public/ShaderCompilerCore.h
--- a/Engine/Source/Runtime/RenderCore/Public/ShaderCompilerCore.h
+++ b/Engine/Source/Runtime/RenderCore/Public/ShaderCompilerCore.h
@@ -16,7 +16,7 @@
 class Error;
 
 // this is for the protocol, not the data, bump if FShaderCompilerInput or ProcessInputFromArchive changes.
-const int32 ShaderCompileWorkerInputVersion = 14;
+const int32 ShaderCompileWorkerInputVersion = 15;
 // this is for the protocol, not the data, bump if FShaderCompilerOutput or WriteToOutputArchive changes.
 const int32 ShaderCompileWorkerOutputVersion = 7;
 // this is for the protocol, not the data.
@@ -173,6 +173,7 @@ struct FShaderCompilerInput
 	FString SourceFilePrefix;
 	FString VirtualSourceFilePath;
 	FString EntryPointName;
+	FString ShaderName;
 
 	// Skips the preprocessor and instead loads the usf file directly
 	bool bSkipPreprocessedCache;
diff --git a/Engine/Source/Runtime/RenderCore/Public/ShaderCore.h b/Engine/Source/Runtime/RenderCore/Public/ShaderCore.h
--- a/Engine/Source/Runtime/RenderCore/Public/ShaderCore.h
+++ b/Engine/Source/Runtime/RenderCore/Public/ShaderCore.h
@@ -919,6 +919,10 @@ extern void GenerateReferencedUniformBuffers(
 /** Records information about all the uniform buffer layouts referenced by UniformBufferEntries. */
 extern RENDERCORE_API void SerializeUniformBufferInfo(class FShaderSaveArchive& Ar, const TMap<const TCHAR*,FCachedUniformBufferDeclaration>& UniformBufferEntries);
 
+/** Create a block of source code to be injected in the preprocessed shader code. The Block will be put into a #line directive
+ * to show up in case shader compilation failures happen in this code block.
+ */
+FString RENDERCORE_API MakeInjectedShaderCodeBlock(const TCHAR* BlockName, const FString& CodeToInject);
 
 
 /**
diff --git a/Engine/Source/Runtime/RenderCore/Public/ShaderParameterMetadata.h b/Engine/Source/Runtime/RenderCore/Public/ShaderParameterMetadata.h
--- a/Engine/Source/Runtime/RenderCore/Public/ShaderParameterMetadata.h
+++ b/Engine/Source/Runtime/RenderCore/Public/ShaderParameterMetadata.h
@@ -125,7 +125,10 @@ public:
 
 	/** Shader binding name of the uniform buffer that contains the root shader parameters. */
 	static constexpr const TCHAR* kRootUniformBufferBindingName = TEXT("_RootShaderParameters");
-	
+
+	/** Shader binding name of the uniform buffer that contains the root shader parameters. */
+	static constexpr int32 kRootCBufferBindingIndex = 0;
+
 	/** A member of a shader parameter structure. */
 	class RENDERCORE_API FMember
 	{
diff --git a/Engine/Source/Runtime/Renderer/Private/GlobalDistanceField.cpp b/Engine/Source/Runtime/Renderer/Private/GlobalDistanceField.cpp
--- a/Engine/Source/Runtime/Renderer/Private/GlobalDistanceField.cpp
+++ b/Engine/Source/Runtime/Renderer/Private/GlobalDistanceField.cpp
@@ -1252,7 +1252,6 @@ class FComposeObjectsIntoPagesCS : public FGlobalShader
 		SHADER_PARAMETER(float, ClipmapVoxelExtent)
 		SHADER_PARAMETER(float, InfluenceRadius)
 		SHADER_PARAMETER(float, InfluenceRadiusSq)
-		SHADER_PARAMETER(uint32, NumSceneObjects)
 		SHADER_PARAMETER(FIntVector, CullGridResolution)
 		SHADER_PARAMETER(FIntVector, GlobalDistanceFieldScrollOffset)
 		SHADER_PARAMETER(FVector3f, GlobalDistanceFieldInvPageAtlasSize)
