Fixes shader compile worker crash on vulkan

diff --git a/Engine/Source/Developer/ShaderCompilerCommon/Private/ShaderCompilerCommon.cpp b/Engine/Source/Developer/ShaderCompilerCommon/Private/ShaderCompilerCommon.cpp
--- a/Engine/Source/Developer/ShaderCompilerCommon/Private/ShaderCompilerCommon.cpp
+++ b/Engine/Source/Developer/ShaderCompilerCommon/Private/ShaderCompilerCommon.cpp
@@ -1071,7 +1071,9 @@ bool FShaderParameterParser::ParseAndMoveShaderParametersToRootConstantBuffer(
 			PreprocessedShaderSource);
 
 		PreprocessedShaderSource = MoveTemp(NewShaderCode);
-	}
+
+		bMovedLoosedParametersToRootConstantBuffer = true;
+	} // if (bMoveToRootConstantBuffer)
 
 	return bSuccess;
 }
@@ -1088,9 +1090,9 @@ void FShaderParameterParser::ValidateShaderParameterType(
 	check(ParsedParameter.IsFound());
 	check(CompilerInput.RootParametersStructure);
 
-	if (ReflectionSize > 0)
+	if (ReflectionSize > 0 && bMovedLoosedParametersToRootConstantBuffer)
 	{
-		// Verify the offset of the parameter comming from shader reflections matches
+		// Verify the offset of the parameter coming from shader reflections honor the packoffset()
 		check(ReflectionOffset == ParsedParameter.ConstantBufferOffset);
 	}
 
diff --git a/Engine/Source/Developer/ShaderCompilerCommon/Public/ShaderCompilerCommon.h b/Engine/Source/Developer/ShaderCompilerCommon/Public/ShaderCompilerCommon.h
--- a/Engine/Source/Developer/ShaderCompilerCommon/Public/ShaderCompilerCommon.h
+++ b/Engine/Source/Developer/ShaderCompilerCommon/Public/ShaderCompilerCommon.h
@@ -117,6 +117,8 @@ private:
 	FString OriginalParsedShader;
 
 	TMap<FString, FParsedShaderParameter> ParsedParameters;
+
+	bool bMovedLoosedParametersToRootConstantBuffer = false;
 };
 
 /** Adds a note to CompilerOutput.Error about where the shader parameter structure is on C++ side. */
