Cherry pick Vulkan shader stripping fix (merged from CL 17822037).

#rnx

diff --git a/Engine/Source/Developer/ShaderCompilerCommon/Private/HlslUtils.cpp b/Engine/Source/Developer/ShaderCompilerCommon/Private/HlslUtils.cpp
--- a/Engine/Source/Developer/ShaderCompilerCommon/Private/HlslUtils.cpp
+++ b/Engine/Source/Developer/ShaderCompilerCommon/Private/HlslUtils.cpp
@@ -1414,12 +1414,6 @@ bool RemoveUnusedInputs(FString& InOutSourceCode, const TArray<FString>& InInput
 	return false;
 }
 
-void StripInstancedStereo(FString& ShaderSource)
-{
-	ShaderSource.ReplaceInline(TEXT("ResolvedView = ResolveView();"), TEXT(""));
-	ShaderSource.ReplaceInline(TEXT("ResolvedView"), TEXT("View"));
-}
-
 struct FConvertFP32ToFP16 {
 	FString Filename;
 	FString GeneratedCode;
diff --git a/Engine/Source/Developer/ShaderCompilerCommon/Public/ShaderCompilerCommon.h b/Engine/Source/Developer/ShaderCompilerCommon/Public/ShaderCompilerCommon.h
--- a/Engine/Source/Developer/ShaderCompilerCommon/Public/ShaderCompilerCommon.h
+++ b/Engine/Source/Developer/ShaderCompilerCommon/Public/ShaderCompilerCommon.h
@@ -148,12 +148,6 @@ extern SHADERCOMPILERCOMMON_API bool RemoveUnusedInputs(FString& InOutSourceCode
 
 extern SHADERCOMPILERCOMMON_API bool ConvertFromFP32ToFP16(FString& InOutSourceCode, TArray<FString>& OutErrors);
 
-/**
-* Fall back to using the View uniform buffer directly for platforms that don't support instanced stereo.
-* @param ShaderSource - Preprocessed shader source
-*/
-extern SHADERCOMPILERCOMMON_API void StripInstancedStereo(FString& ShaderSource);
-
 extern SHADERCOMPILERCOMMON_API FString CreateShaderCompilerWorkerDirectCommandLine(const FShaderCompilerInput& Input, uint32 CCFlags = 0);
 
 enum class EShaderConductorTarget
diff --git a/Engine/Source/Developer/VulkanShaderFormat/Private/VulkanShaderCompiler.cpp b/Engine/Source/Developer/VulkanShaderFormat/Private/VulkanShaderCompiler.cpp
--- a/Engine/Source/Developer/VulkanShaderFormat/Private/VulkanShaderCompiler.cpp
+++ b/Engine/Source/Developer/VulkanShaderFormat/Private/VulkanShaderCompiler.cpp
@@ -2106,12 +2106,6 @@ void DoCompileVulkanShader(const FShaderCompilerInput& Input, FShaderCompilerOut
 		}
 	}
 
-	if (!(PreprocessedShaderSource.Contains(TEXT("SV_ViewID")) || PreprocessedShaderSource.Contains(TEXT("VIEW_ID"))))
-	{
-		// Disable instanced stereo if not requested by the shader
-		StripInstancedStereo(PreprocessedShaderSource);
-	}
-
 	FShaderParameterParser ShaderParameterParser;
 	if (!ShaderParameterParser.ParseAndMoveShaderParametersToRootConstantBuffer(
 		Input, Output, PreprocessedShaderSource, /* ConstantBufferType = */ nullptr))
