Move DebugDumpBlobDetail function to public header to expand use to Windows.

#rnx

diff --git a/Engine/Source/Developer/ShaderCompilerCommon/Public/DXCUtils.inl b/Engine/Source/Developer/ShaderCompilerCommon/Public/DXCUtils.inl
new file mode 100644
--- /dev/null
+++ b/Engine/Source/Developer/ShaderCompilerCommon/Public/DXCUtils.inl
@@ -0,0 +1,61 @@
+// Copyright Epic Games, Inc. All Rights Reserved.
+
+#pragma once
+
+static void DumpDebugBlobDetail(IDxcContainerReflection* Reflection, IDxcBlob* Blob, const TCHAR* BlobName)
+{
+	check(Reflection != nullptr);
+	check(BlobName != nullptr);
+
+	if (!FPlatformMisc::IsDebuggerPresent())
+	{
+		return;
+	}
+
+	if (Blob == nullptr)
+	{
+		FPlatformMisc::LowLevelOutputDebugStringf(TEXT("Blob %s is NULL\n"), BlobName);
+		return;
+	}
+
+	Reflection->Load(Blob);
+
+	UINT32 PartCount = 0;
+	Reflection->GetPartCount(&PartCount);
+
+	FPlatformMisc::LowLevelOutputDebugStringf(TEXT("Blob %s : %d parts, %d total bytes\n"), BlobName, PartCount, Blob->GetBufferSize());
+	for (UINT32 PartIndex = 0; PartIndex < PartCount; PartIndex++)
+	{
+		UINT32 PartKind;
+		Reflection->GetPartKind(PartIndex, &PartKind);
+
+		const TCHAR* BlobDescription;
+		switch (PartKind)
+		{
+		case DXC_FOURCC('D', 'X', 'B', 'C'): BlobDescription = TEXT("Container (back-compat)"); break;
+		case DXC_FOURCC('R', 'D', 'E', 'F'): BlobDescription = TEXT("ResourceDef");             break;
+		case DXC_FOURCC('I', 'S', 'G', '1'): BlobDescription = TEXT("InputSignature");          break;
+		case DXC_FOURCC('O', 'S', 'G', '1'): BlobDescription = TEXT("OutputSignature");         break;
+		case DXC_FOURCC('P', 'S', 'G', '1'): BlobDescription = TEXT("PatchConstantSignature");  break;
+		case DXC_FOURCC('S', 'T', 'A', 'T'): BlobDescription = TEXT("ShaderStatistics");        break;
+		case DXC_FOURCC('I', 'L', 'D', 'B'): BlobDescription = TEXT("ShaderDebugInfoDXIL");     break;
+		case DXC_FOURCC('I', 'L', 'D', 'N'): BlobDescription = TEXT("ShaderDebugName");         break;
+		case DXC_FOURCC('S', 'F', 'I', '0'): BlobDescription = TEXT("FeatureInfo");             break;
+		case DXC_FOURCC('P', 'R', 'I', 'V'): BlobDescription = TEXT("PrivateData");             break;
+		case DXC_FOURCC('R', 'T', 'S', '0'): BlobDescription = TEXT("RootSignature");           break;
+		case DXC_FOURCC('D', 'X', 'I', 'L'): BlobDescription = TEXT("DXIL");                    break;
+		case DXC_FOURCC('P', 'S', 'V', '0'): BlobDescription = TEXT("PipelineStateValidation"); break;
+		case DXC_FOURCC('R', 'D', 'A', 'T'): BlobDescription = TEXT("RuntimeData");             break;
+		case DXC_FOURCC('H', 'A', 'S', 'H'): BlobDescription = TEXT("ShaderHash");              break;
+		default:                             BlobDescription = TEXT("(unknown)");               break;
+		}
+
+		TRefCountPtr<IDxcBlob> PartBlob;
+		Reflection->GetPartContent(PartIndex, PartBlob.GetInitReference());
+
+		FPlatformMisc::LowLevelOutputDebugStringf(TEXT("\t[%4d] %c%c%c%c (%s) %d bytes\n"), PartIndex, (PartKind & 0xFF), (PartKind >> 8) & 0xFF, (PartKind >> 16) & 0xFF, (PartKind >> 24) & 0xFF, BlobDescription, PartBlob->GetBufferSize());
+	}
+	FPlatformMisc::LowLevelOutputDebugString(TEXT("\n\n"));
+}
+
+
