Fix compile error introduced in CL 18033267 and adopt temporary reflection object in DumpDebugBlobDetail from //UE5/Release-Engine-Staging.

#rnx

diff --git a/Engine/Source/Developer/ShaderCompilerCommon/Public/DXCUtils.inl b/Engine/Source/Developer/ShaderCompilerCommon/Public/DXCUtils.inl
--- a/Engine/Source/Developer/ShaderCompilerCommon/Public/DXCUtils.inl
+++ b/Engine/Source/Developer/ShaderCompilerCommon/Public/DXCUtils.inl
@@ -2,9 +2,8 @@
 
 #pragma once
 
-static void DumpDebugBlobDetail(IDxcContainerReflection* Reflection, IDxcBlob* Blob, const TCHAR* BlobName)
+static void DumpDebugBlobDetail(IDxcBlob* Blob, const TCHAR* BlobName)
 {
-	check(Reflection != nullptr);
 	check(BlobName != nullptr);
 
 	if (!FPlatformMisc::IsDebuggerPresent())
@@ -18,6 +17,14 @@ static void DumpDebugBlobDetail(IDxcContainerReflection* Reflection, IDxcBlob* B
 		return;
 	}
 
+	TRefCountPtr<IDxcContainerReflection> Reflection;
+	DxcCreateInstance(CLSID_DxcContainerReflection, __uuidof(IDxcContainerReflection), (void**)Reflection.GetInitReference());
+	if (!Reflection.IsValid())
+	{
+		FPlatformMisc::LowLevelOutputDebugStringf(TEXT("Can't show debug detail for Blob %s - Failed to create IDxcContainerReflection\n"), BlobName);
+		return;
+	}
+
 	Reflection->Load(Blob);
 
 	UINT32 PartCount = 0;
