Enabled half precision support for OpenGL and Vulkan DXC

diff --git a/Engine/Shaders/Public/Platform.ush b/Engine/Shaders/Public/Platform.ush
--- a/Engine/Shaders/Public/Platform.ush
+++ b/Engine/Shaders/Public/Platform.ush
@@ -207,6 +207,15 @@
 	#define fixed3x3 float3x3
 	#define fixed4x4 float4x4
 	#define fixed4x3 float4x3
+#elif (VULKAN_PROFILE || COMPILER_GLSL_ES3_1) && !COMPILER_HLSLCC
+	// For VULKAN and OPENGL ES31 use RelaxedPrecision for half floats
+	#define half min16float
+	#define half2 min16float2
+	#define half3 min16float3
+	#define half4 min16float4
+	#define half3x3 min16float3x3
+	#define half3x4 min16float3x4 
+	#define half4x4 min16float4x4
 #endif
 
 // ---------------------------------------------------- Profile config
@@ -463,7 +472,7 @@ void ShaderYield()
 #define COMPILER_SUPPORTS_GATHER_UINT 0
 #endif
 
-#if !COMPILER_HLSLCC && ES3_1_PROFILE
+#if ES3_1_PROFILE
 #define HALF_TYPE half
 #define HALF2_TYPE half2
 #define HALF3_TYPE half3
diff --git a/Engine/Source/Developer/ShaderCompilerCommon/Private/ShaderCompilerCommon.cpp b/Engine/Source/Developer/ShaderCompilerCommon/Private/ShaderCompilerCommon.cpp
--- a/Engine/Source/Developer/ShaderCompilerCommon/Private/ShaderCompilerCommon.cpp
+++ b/Engine/Source/Developer/ShaderCompilerCommon/Private/ShaderCompilerCommon.cpp
@@ -1114,6 +1114,15 @@ void FShaderParameterParser::ValidateShaderParameterType(
 			{
 				bIsTypeCorrect = (FCString::Strcmp(*ParsedParameter.ParsedType + 5, *ExpectedShaderType + 4) == 0);
 			}
+			// support for min16float
+			else if (ParsedParameter.ParsedType.StartsWith(TEXT("min16float")) && ExpectedShaderType.StartsWith(TEXT("float")))
+			{
+				bIsTypeCorrect = (FCString::Strcmp(*ParsedParameter.ParsedType + 10, *ExpectedShaderType + 5) == 0);
+			}
+			else if (ParsedParameter.ParsedType.StartsWith(TEXT("min16float")) && ExpectedShaderType.StartsWith(TEXT("half")))
+			{
+				bIsTypeCorrect = (FCString::Strcmp(*ParsedParameter.ParsedType + 10, *ExpectedShaderType + 4) == 0);
+			}
 		}
 
 		// Allow silent casting between signed and unsigned on shader bindings.
