Removing usage of DDSPI from ValidateShaderParameterTypes

diff --git a/Engine/Source/Developer/Apple/MetalShaderFormat/Private/MetalShaderCompiler.cpp b/Engine/Source/Developer/Apple/MetalShaderFormat/Private/MetalShaderCompiler.cpp
--- a/Engine/Source/Developer/Apple/MetalShaderFormat/Private/MetalShaderCompiler.cpp
+++ b/Engine/Source/Developer/Apple/MetalShaderFormat/Private/MetalShaderCompiler.cpp
@@ -1169,7 +1169,7 @@ void CompileShader_Metal(const FShaderCompilerInput& _Input,FShaderCompilerOutpu
 
 	}
 
-	ShaderParameterParser.ValidateShaderParameterTypes(Input, Output);
+	ShaderParameterParser.ValidateShaderParameterTypes(Input, bIsMobile, Output);
 }
 
 bool StripShader_Metal(TArray<uint8>& Code, class FString const& DebugPath, bool const bNative)
diff --git a/Engine/Source/Developer/ShaderCompilerCommon/Private/ShaderCompilerCommon.cpp b/Engine/Source/Developer/ShaderCompilerCommon/Private/ShaderCompilerCommon.cpp
--- a/Engine/Source/Developer/ShaderCompilerCommon/Private/ShaderCompilerCommon.cpp
+++ b/Engine/Source/Developer/ShaderCompilerCommon/Private/ShaderCompilerCommon.cpp
@@ -1076,6 +1076,7 @@ void FShaderParameterParser::ValidateShaderParameterType(
 	const FString& ShaderBindingName,
 	int32 ReflectionOffset,
 	int32 ReflectionSize,
+	bool bPlatformSupportsPrecisionModifier,
 	FShaderCompilerOutput& CompilerOutput) const
 {
 	const FShaderParameterParser::FParsedShaderParameter& ParsedParameter = FindParameterInfos(ShaderBindingName);
@@ -1089,12 +1090,10 @@ void FShaderParameterParser::ValidateShaderParameterType(
 		check(ReflectionOffset == ParsedParameter.ConstantBufferOffset);
 	}
 
-	EShaderPlatform ShaderPlatform = EShaderPlatform(CompilerInput.Target.Platform);
-
 	// Validate the shader type.
 	{
 		FString ExpectedShaderType;
-		ParsedParameter.Member->GenerateShaderParameterType(ExpectedShaderType, ShaderPlatform);
+		ParsedParameter.Member->GenerateShaderParameterType(ExpectedShaderType, bPlatformSupportsPrecisionModifier);
 
 		const bool bShouldBeInt = ParsedParameter.Member->GetBaseType() == UBMT_INT32;
 		const bool bShouldBeUint = ParsedParameter.Member->GetBaseType() == UBMT_UINT32;
@@ -1187,6 +1186,7 @@ void FShaderParameterParser::ValidateShaderParameterType(
 
 void FShaderParameterParser::ValidateShaderParameterTypes(
 	const FShaderCompilerInput& CompilerInput,
+	bool bPlatformSupportsPrecisionModifier,
 	FShaderCompilerOutput& CompilerOutput) const
 {
 	// The shader doesn't have any parameter binding through shader structure, therefore don't do anything.
@@ -1238,7 +1238,7 @@ void FShaderParameterParser::ValidateShaderParameterTypes(
 			BoundSize = ParameterAllocation->Size;
 		}
 
-		ValidateShaderParameterType(CompilerInput, ShaderBindingName, BoundOffset, BoundSize, CompilerOutput);
+		ValidateShaderParameterType(CompilerInput, ShaderBindingName, BoundOffset, BoundSize, bPlatformSupportsPrecisionModifier, CompilerOutput);
 	});
 }
 
diff --git a/Engine/Source/Developer/ShaderCompilerCommon/Public/ShaderCompilerCommon.h b/Engine/Source/Developer/ShaderCompilerCommon/Public/ShaderCompilerCommon.h
--- a/Engine/Source/Developer/ShaderCompilerCommon/Public/ShaderCompilerCommon.h
+++ b/Engine/Source/Developer/ShaderCompilerCommon/Public/ShaderCompilerCommon.h
@@ -3,7 +3,6 @@
 #pragma once
 
 #include "CoreMinimal.h"
-#include "UniformBuffer.h"
 #include "ShaderCore.h"
 #include "ShaderCompilerCore.h"
 #include "CrossCompilerDefinitions.h"
@@ -98,13 +97,32 @@ public:
 		const FString& ShaderBindingName,
 		int32 ReflectionOffset,
 		int32 ReflectionSize,
+		bool bPlatformSupportsPrecisionModifier,
 		FShaderCompilerOutput& CompilerOutput) const;
 
+	void ValidateShaderParameterType(
+		const FShaderCompilerInput& CompilerInput,
+		const FString& ShaderBindingName,
+		int32 ReflectionOffset,
+		int32 ReflectionSize,
+		FShaderCompilerOutput& CompilerOutput) const
+	{
+		ValidateShaderParameterType(CompilerInput, ShaderBindingName, ReflectionOffset, ReflectionSize, false, CompilerOutput);
+	}
+
 	/** Validates shader parameter map is compatible with the shader parameter structure. */
 	void ValidateShaderParameterTypes(
 		const FShaderCompilerInput& CompilerInput,
+		bool bPlatformSupportsPrecisionModifier,
 		FShaderCompilerOutput& CompilerOutput) const;
 
+	void ValidateShaderParameterTypes(
+		const FShaderCompilerInput& CompilerInput,
+		FShaderCompilerOutput& CompilerOutput) const
+	{
+		ValidateShaderParameterTypes(CompilerInput, false, CompilerOutput);
+	}
+
 	/** Gets file and line of the parameter in the shader source code. */
 	void GetParameterFileAndLine(const FParsedShaderParameter& ParsedParameter, FString& OutFile, FString& OutLine) const
 	{
diff --git a/Engine/Source/Developer/VulkanShaderFormat/Private/VulkanShaderCompiler.cpp b/Engine/Source/Developer/VulkanShaderFormat/Private/VulkanShaderCompiler.cpp
--- a/Engine/Source/Developer/VulkanShaderFormat/Private/VulkanShaderCompiler.cpp
+++ b/Engine/Source/Developer/VulkanShaderFormat/Private/VulkanShaderCompiler.cpp
@@ -2188,7 +2188,7 @@ void DoCompileVulkanShader(const FShaderCompilerInput& Input, FShaderCompilerOut
 	bSuccess = CompileWithShaderConductor(PreprocessedShaderSource, EntryPointName, Frequency, CompilerInfo, Output, BindingTable, bStripReflect);
 #endif // PLATFORM_MAC || PLATFORM_WINDOWS || PLATFORM_LINUX
 	
-	ShaderParameterParser.ValidateShaderParameterTypes(Input, Output);
+	ShaderParameterParser.ValidateShaderParameterTypes(Input, bIsMobile, Output);
 	
 	if (bDirectCompile)
 	{
diff --git a/Engine/Source/Runtime/RenderCore/Private/ShaderParameterMetadata.cpp b/Engine/Source/Runtime/RenderCore/Private/ShaderParameterMetadata.cpp
--- a/Engine/Source/Runtime/RenderCore/Private/ShaderParameterMetadata.cpp
+++ b/Engine/Source/Runtime/RenderCore/Private/ShaderParameterMetadata.cpp
@@ -59,14 +59,14 @@ TMap<FHashedName, FShaderParametersMetadata*>& FShaderParametersMetadata::GetNam
 	return NameStructMap;
 }
 
-void FShaderParametersMetadata::FMember::GenerateShaderParameterType(FString& Result, EShaderPlatform ShaderPlatform) const
+void FShaderParametersMetadata::FMember::GenerateShaderParameterType(FString& Result, bool bSupportsPrecisionModifier) const
 {
 	switch (GetBaseType())
 	{
 	case UBMT_INT32:   Result = TEXT("int"); break;
 	case UBMT_UINT32:  Result = TEXT("uint"); break;
 	case UBMT_FLOAT32:
-		if (GetPrecision() == EShaderPrecisionModifier::Float || !SupportShaderPrecisionModifier(ShaderPlatform))
+		if (GetPrecision() == EShaderPrecisionModifier::Float || !bSupportsPrecisionModifier)
 		{
 			Result = TEXT("float");
 		}
@@ -94,6 +94,11 @@ void FShaderParametersMetadata::FMember::GenerateShaderParameterType(FString& Re
 	}
 }
 
+void FShaderParametersMetadata::FMember::GenerateShaderParameterType(FString& Result, EShaderPlatform ShaderPlatform) const
+{
+	GenerateShaderParameterType(Result, SupportShaderPrecisionModifier(ShaderPlatform));
+}
+
 FShaderParametersMetadata* FindUniformBufferStructByName(const TCHAR* StructName)
 {
 	return FindUniformBufferStructByFName(FName(StructName, FNAME_Find));
diff --git a/Engine/Source/Runtime/RenderCore/Public/ShaderParameterMetadata.h b/Engine/Source/Runtime/RenderCore/Public/ShaderParameterMetadata.h
--- a/Engine/Source/Runtime/RenderCore/Public/ShaderParameterMetadata.h
+++ b/Engine/Source/Runtime/RenderCore/Public/ShaderParameterMetadata.h
@@ -205,6 +205,7 @@ public:
 			return ElementSize;
 		}
 
+		void GenerateShaderParameterType(FString& Result, bool bSupportsPrecisionModifier) const;
 		void GenerateShaderParameterType(FString& Result, EShaderPlatform ShaderPlatform) const;
 
 	private:
