Fix for previous dxc shader compiler integration from UE5 clashing with the DumpDebugBlobDetail changes from CL 18033267.

#rnx

diff --git a/Engine/Source/Developer/ShaderCompilerCommon/Public/DXCUtils.inl b/Engine/Source/Developer/ShaderCompilerCommon/Public/DXCUtils.inl
--- a/Engine/Source/Developer/ShaderCompilerCommon/Public/DXCUtils.inl
+++ b/Engine/Source/Developer/ShaderCompilerCommon/Public/DXCUtils.inl
@@ -2,7 +2,16 @@
 
 #pragma once
 
+// not every shader format compiler dll will have access to DxcCreateInstance directly - some might be relying on GetProcAddress to find it
+#ifndef DXCUTILS_HAS_DXCCREATEINSTANCE
+	#define DXCUTILS_HAS_DXCCREATEINSTANCE 1
+#endif
+
+#if DXCUTILS_HAS_DXCCREATEINSTANCE
 static void DumpDebugBlobDetail(IDxcBlob* Blob, const TCHAR* BlobName)
+#else
+static void DumpDebugBlobDetail(IDxcBlob* Blob, const TCHAR* BlobName, TRefCountPtr<IDxcContainerReflection> Reflection)
+#endif
 {
 	check(BlobName != nullptr);
 
@@ -17,8 +26,10 @@ static void DumpDebugBlobDetail(IDxcBlob* Blob, const TCHAR* BlobName)
 		return;
 	}
 
+#if DXCUTILS_HAS_DXCCREATEINSTANCE
 	TRefCountPtr<IDxcContainerReflection> Reflection;
 	DxcCreateInstance(CLSID_DxcContainerReflection, __uuidof(IDxcContainerReflection), (void**)Reflection.GetInitReference());
+#endif
 	if (!Reflection.IsValid())
 	{
 		FPlatformMisc::LowLevelOutputDebugStringf(TEXT("Can't show debug detail for Blob %s - Failed to create IDxcContainerReflection\n"), BlobName);
