Fix shader parameter parser: Also erase 'const'-qualifier alongside global parameter.

#rnx

diff --git a/Engine/Shaders/Public/ShaderVersion.ush b/Engine/Shaders/Public/ShaderVersion.ush
--- a/Engine/Shaders/Public/ShaderVersion.ush
+++ b/Engine/Shaders/Public/ShaderVersion.ush
@@ -2,5 +2,4 @@
 // Each time the console command is executed it generates a new GUID. As this file is included
 // in Platform.ush (which should be included in any shader) it allows to invalidate the shader DDC.
 // 
-// GUID = B6B2FB8E-934E-4972-B768-7BD0AAF512F1
-
+// GUID = 6B4506D7-B1FE-40E7-8003-9025BBA76A56
diff --git a/Engine/Source/Developer/ShaderCompilerCommon/Private/ShaderCompilerCommon.cpp b/Engine/Source/Developer/ShaderCompilerCommon/Private/ShaderCompilerCommon.cpp
--- a/Engine/Source/Developer/ShaderCompilerCommon/Private/ShaderCompilerCommon.cpp
+++ b/Engine/Source/Developer/ShaderCompilerCommon/Private/ShaderCompilerCommon.cpp
@@ -603,6 +603,7 @@ bool FShaderParameterParser::ParseAndMoveShaderParametersToRootConstantBuffer(
 		int32 CurrentPragamLineoffset = -1;
 		int32 CurrentLineoffset = 0;
 
+		int32 TypeQualifierStartPos = -1;
 		int32 TypeStartPos = -1;
 		int32 TypeEndPos = -1;
 		int32 NameStartPos = -1;
@@ -616,6 +617,7 @@ bool FShaderParameterParser::ParseAndMoveShaderParametersToRootConstantBuffer(
 
 		auto ResetState = [&]()
 		{
+			TypeQualifierStartPos = -1;
 			TypeStartPos = -1;
 			TypeEndPos = -1;
 			NameStartPos = -1;
@@ -702,7 +704,7 @@ bool FShaderParameterParser::ParseAndMoveShaderParametersToRootConstantBuffer(
 				// Erases this shader parameter conserving the same line numbers.
 				if (bEraseOriginalParameter)
 				{
-					for (int32 j = TypeStartPos; j <= Cursor; j++)
+					for (int32 j = (TypeQualifierStartPos != -1 ? TypeQualifierStartPos : TypeStartPos); j <= Cursor; j++)
 					{
 						if (PreprocessedShaderSource[j] != '\r' && PreprocessedShaderSource[j] != '\n')
 							PreprocessedShaderSource[j] = ' ';
@@ -801,6 +803,11 @@ bool FShaderParameterParser::ParseAndMoveShaderParametersToRootConstantBuffer(
 					else if (RecognisedKeywordId == 2)
 					{
 						// Ignore the const keywords, but still parse given it might still be a shader parameter.
+						if (TypeQualifierStartPos == -1)
+						{
+							// If the parameter is erased, we also have to erase *all* 'const'-qualifiers, e.g. "const int Foo" or "const const int Foo".
+							TypeQualifierStartPos = Cursor;
+						}
 						Cursor += KeywordTableSize[RecognisedKeywordId];
 					}
 					else
