RemoveUniformBuffersFromSource shouldn't implicitly call TransformStringIntoCharacterArray.

diff --git a/Engine/Source/Developer/Apple/MetalShaderFormat/Private/MetalShaderCompiler.cpp b/Engine/Source/Developer/Apple/MetalShaderFormat/Private/MetalShaderCompiler.cpp
--- a/Engine/Source/Developer/Apple/MetalShaderFormat/Private/MetalShaderCompiler.cpp
+++ b/Engine/Source/Developer/Apple/MetalShaderFormat/Private/MetalShaderCompiler.cpp
@@ -1057,6 +1057,9 @@ void CompileShader_Metal(const FShaderCompilerInput& _Input,FShaderCompilerOutpu
 
 	// This requires removing the HLSLCC_NoPreprocess flag later on!
 	RemoveUniformBuffersFromSource(Input.Environment, PreprocessedShader);
+
+	// Process TEXT macro.
+	TransformStringIntoCharacterArray(PreprocessedShader);
 	
 	uint32 CCFlags = HLSLCC_NoPreprocess | HLSLCC_PackUniformsIntoUniformBufferWithNames | HLSLCC_FixAtomicReferences | HLSLCC_RetainSizes | HLSLCC_KeepSamplerAndImageNames;
 	if (!bDirectCompile || UE_BUILD_DEBUG)
diff --git a/Engine/Source/Developer/ShaderCompilerCommon/Private/ShaderCompilerCommon.cpp b/Engine/Source/Developer/ShaderCompilerCommon/Private/ShaderCompilerCommon.cpp
--- a/Engine/Source/Developer/ShaderCompilerCommon/Private/ShaderCompilerCommon.cpp
+++ b/Engine/Source/Developer/ShaderCompilerCommon/Private/ShaderCompilerCommon.cpp
@@ -1309,8 +1309,6 @@ void FShaderParameterParser::ExtractFileAndLine(int32 PragamLineoffset, int32 Li
 	OutLine = FString::FromInt(FinalLine);
 }
 
-void TransformStringIntoCharacterArray(FString& PreprocessedShaderSource);
-
 // The cross compiler doesn't yet support struct initializers needed to construct static structs for uniform buffers
 // Replace all uniform buffer struct member references (View.WorldToClip) with a flattened name that removes the struct dependency (View_WorldToClip)
 void RemoveUniformBuffersFromSource(const FShaderCompilerEnvironment& Environment, FString& PreprocessedShaderSource)
@@ -1400,12 +1398,6 @@ void RemoveUniformBuffersFromSource(const FShaderCompilerEnvironment& Environmen
 			SearchPtr = FindNextUniformBufferReference(SearchPtr + UniformBufferAccessString.Len(), *UniformBufferName, UniformBufferName.Len());
 		}
 	}
-
-	// Process TEXT macro.
-	// Even this processing is not directly related to the Uniform buffer removal, this processing is 
-	// called here as this function is called by all the platform specific compiler, and avoid to update 
-	// all compiler
-	TransformStringIntoCharacterArray(PreprocessedShaderSource);
 }
 
 ///////////////////////////////////////////////////////////////////////////////////////////////////////////////////////////////////////
diff --git a/Engine/Source/Developer/ShaderCompilerCommon/Public/ShaderCompilerCommon.h b/Engine/Source/Developer/ShaderCompilerCommon/Public/ShaderCompilerCommon.h
--- a/Engine/Source/Developer/ShaderCompilerCommon/Public/ShaderCompilerCommon.h
+++ b/Engine/Source/Developer/ShaderCompilerCommon/Public/ShaderCompilerCommon.h
@@ -160,6 +160,10 @@ extern SHADERCOMPILERCOMMON_API void AddUnboundShaderParameterError(
 // The cross compiler doesn't yet support struct initializers needed to construct static structs for uniform buffers
 // Replace all uniform buffer struct member references (View.WorldToClip) with a flattened name that removes the struct dependency (View_WorldToClip)
 extern SHADERCOMPILERCOMMON_API void RemoveUniformBuffersFromSource(const FShaderCompilerEnvironment& Environment, FString& PreprocessedShaderSource);
+
+// Processes TEXT macros
+extern SHADERCOMPILERCOMMON_API void TransformStringIntoCharacterArray(FString& PreprocessedShaderSource);
+
 extern SHADERCOMPILERCOMMON_API bool RemoveUnusedOutputs(FString& InOutSourceCode, const TArray<FString>& InUsedOutputs, const TArray<FString>& InExceptions, FString& InOutEntryPoint, TArray<FString>& OutErrors);
 
 extern SHADERCOMPILERCOMMON_API bool RemoveUnusedInputs(FString& InOutSourceCode, const TArray<FString>& InUsedInputs, FString& InOutEntryPoint, TArray<FString>& OutErrors);
diff --git a/Engine/Source/Developer/ShaderFormatOpenGL/Private/OpenGLShaderCompiler.cpp b/Engine/Source/Developer/ShaderFormatOpenGL/Private/OpenGLShaderCompiler.cpp
--- a/Engine/Source/Developer/ShaderFormatOpenGL/Private/OpenGLShaderCompiler.cpp
+++ b/Engine/Source/Developer/ShaderFormatOpenGL/Private/OpenGLShaderCompiler.cpp
@@ -3213,6 +3213,9 @@ void FOpenGLFrontend::CompileShader(const FShaderCompilerInput& Input, FShaderCo
 	// This requires removing the HLSLCC_NoPreprocess flag later on!
 	RemoveUniformBuffersFromSource(Input.Environment, PreprocessedShader);
 
+	// Process TEXT macro.
+	TransformStringIntoCharacterArray(PreprocessedShader);
+
 	uint32 CCFlags = CalculateCrossCompilerFlags(Version, Input.Environment.FullPrecisionInPS, Input.Environment.CompilerFlags);
 
 	// Required as we added the RemoveUniformBuffersFromSource() function (the cross-compiler won't be able to interpret comments w/o a preprocessor)
diff --git a/Engine/Source/Developer/VulkanShaderFormat/Private/VulkanShaderCompiler.cpp b/Engine/Source/Developer/VulkanShaderFormat/Private/VulkanShaderCompiler.cpp
--- a/Engine/Source/Developer/VulkanShaderFormat/Private/VulkanShaderCompiler.cpp
+++ b/Engine/Source/Developer/VulkanShaderFormat/Private/VulkanShaderCompiler.cpp
@@ -2323,6 +2323,9 @@ void DoCompileVulkanShader(const FShaderCompilerInput& Input, FShaderCompilerOut
 
 	RemoveUniformBuffersFromSource(Input.Environment, PreprocessedShaderSource);
 
+	// Process TEXT macro.
+	TransformStringIntoCharacterArray(PreprocessedShaderSource);
+
 	FCompilerInfo CompilerInfo(Input, WorkingDirectory, HlslFrequency);
 
 	// Setup hlslcc flags. Needed here as it will be used when dumping debug info
diff --git a/Engine/Source/Developer/Windows/ShaderFormatD3D/Private/D3DShaderCompiler.cpp b/Engine/Source/Developer/Windows/ShaderFormatD3D/Private/D3DShaderCompiler.cpp
--- a/Engine/Source/Developer/Windows/ShaderFormatD3D/Private/D3DShaderCompiler.cpp
+++ b/Engine/Source/Developer/Windows/ShaderFormatD3D/Private/D3DShaderCompiler.cpp
@@ -1270,6 +1270,9 @@ void CompileD3DShader(const FShaderCompilerInput& Input, FShaderCompilerOutput&
 
 	RemoveUniformBuffersFromSource(Input.Environment, PreprocessedShaderSource);
 
+	// Process TEXT macro.
+	TransformStringIntoCharacterArray(PreprocessedShaderSource);
+
 	// @TODO - implement different material path to allow us to remove backwards compat flag on sm5 shaders
 	uint32 CompileFlags = D3DCOMPILE_ENABLE_BACKWARDS_COMPATIBILITY
 		// Unpack uniform matrices as row-major to match the CPU layout.
