Migrating back to using UniformBuffer structs.
Since the cbuffer uses the same name as the struct, we have to rename the cbuffer using a consistent prefix. The reflection handling also needs to reinterpret the renamed cbuffer as the correct UniformBuffer.

diff --git a/Engine/Source/Developer/ShaderCompilerCommon/Private/ShaderCompilerCommon.cpp b/Engine/Source/Developer/ShaderCompilerCommon/Private/ShaderCompilerCommon.cpp
--- a/Engine/Source/Developer/ShaderCompilerCommon/Private/ShaderCompilerCommon.cpp
+++ b/Engine/Source/Developer/ShaderCompilerCommon/Private/ShaderCompilerCommon.cpp
@@ -5,6 +5,7 @@
 #include "Modules/ModuleManager.h"
 #include "HlslccDefinitions.h"
 #include "HAL/FileManager.h"
+#include "ShaderParameters.h"
 
 IMPLEMENT_MODULE(FDefaultModuleImpl, ShaderCompilerCommon);
 
@@ -469,6 +470,13 @@ TCHAR* FindNextUniformBufferReference(TCHAR* SearchPtr, const TCHAR* SearchStrin
 	return nullptr;
 }
 
+FString UE::ShaderCompilerCommon::SanitizeUniformBufferName(const FString& InName)
+{
+	FString NewName(InName);
+	NewName.RemoveFromStart(UE::ShaderParameters::kUniformBufferConstantBufferPrefix);
+	return NewName;
+}
+
 void HandleReflectedGlobalConstantBufferMember(
 	const FString& MemberName,
 	uint32 ConstantBufferIndex,
@@ -518,8 +526,10 @@ void HandleReflectedUniformBuffer(
 	FShaderCompilerOutput& CompilerOutput
 )
 {
+	FString AdjustedUniformBufferName(UE::ShaderCompilerCommon::SanitizeUniformBufferName(UniformBufferName));
+
 	CompilerOutput.ParameterMap.AddParameterAllocation(
-		*UniformBufferName,
+		*AdjustedUniformBufferName,
 		ReflectionSlot,
 		BaseIndex,
 		BufferSize,
diff --git a/Engine/Source/Developer/ShaderCompilerCommon/Public/ShaderCompilerCommon.h b/Engine/Source/Developer/ShaderCompilerCommon/Public/ShaderCompilerCommon.h
--- a/Engine/Source/Developer/ShaderCompilerCommon/Public/ShaderCompilerCommon.h
+++ b/Engine/Source/Developer/ShaderCompilerCommon/Public/ShaderCompilerCommon.h
@@ -139,6 +139,11 @@ private:
 	bool bMovedLoosedParametersToRootConstantBuffer = false;
 };
 
+namespace UE::ShaderCompilerCommon
+{
+	extern SHADERCOMPILERCOMMON_API FString SanitizeUniformBufferName(const FString& InName);
+}
+
 extern SHADERCOMPILERCOMMON_API void HandleReflectedGlobalConstantBufferMember(
 	const FString& MemberName,
 	uint32 ConstantBufferIndex,
diff --git a/Engine/Source/Developer/Windows/ShaderFormatD3D/Private/D3DShaderCompiler.cpp b/Engine/Source/Developer/Windows/ShaderFormatD3D/Private/D3DShaderCompiler.cpp
--- a/Engine/Source/Developer/Windows/ShaderFormatD3D/Private/D3DShaderCompiler.cpp
+++ b/Engine/Source/Developer/Windows/ShaderFormatD3D/Private/D3DShaderCompiler.cpp
@@ -1276,7 +1276,11 @@ void CompileD3DShader(const FShaderCompilerInput& Input, FShaderCompilerOutput&
 		return;
 	}
 
-	RemoveUniformBuffersFromSource(Input.Environment, PreprocessedShaderSource);
+	// Only use UniformBuffer structs on SM6 until we can fully vet SM5
+	if (Language != ELanguage::SM6)
+	{
+		RemoveUniformBuffersFromSource(Input.Environment, PreprocessedShaderSource);
+	}
 
 	// Process TEXT macro.
 	TransformStringIntoCharacterArray(PreprocessedShaderSource);
diff --git a/Engine/Source/Developer/Windows/ShaderFormatD3D/Private/D3DShaderCompiler.inl b/Engine/Source/Developer/Windows/ShaderFormatD3D/Private/D3DShaderCompiler.inl
--- a/Engine/Source/Developer/Windows/ShaderFormatD3D/Private/D3DShaderCompiler.inl
+++ b/Engine/Source/Developer/Windows/ShaderFormatD3D/Private/D3DShaderCompiler.inl
@@ -412,7 +412,7 @@ template <typename ID3D1xShaderReflection, typename D3D1x_SHADER_DESC, typename
 				{
 					UniformBufferNames.AddDefaulted(CBIndex - UniformBufferNames.Num() + 1);
 				}
-				UniformBufferNames[CBIndex] = UniformBufferName;
+				UniformBufferNames[CBIndex] = UE::ShaderCompilerCommon::SanitizeUniformBufferName(UniformBufferName);
 				UsedUniformBufferSlots[CBIndex] = true;
 			}
 
diff --git a/Engine/Source/Runtime/RenderCore/Private/ShaderParameters.cpp b/Engine/Source/Runtime/RenderCore/Private/ShaderParameters.cpp
--- a/Engine/Source/Runtime/RenderCore/Private/ShaderParameters.cpp
+++ b/Engine/Source/Runtime/RenderCore/Private/ShaderParameters.cpp
@@ -278,7 +278,7 @@ static FString CreateHLSLUniformBufferDeclaration(const TCHAR* Name,const FShade
 		return FString::Printf(
 			TEXT("#ifndef __UniformBuffer_%s_Definition__\r\n")
 			TEXT("#define __UniformBuffer_%s_Definition__\r\n")
-			TEXT("cbuffer %s\r\n")
+			TEXT("cbuffer %s%s\r\n")
 			TEXT("{\r\n")
 			TEXT("%s")
 			TEXT("}\r\n")
@@ -290,6 +290,7 @@ static FString CreateHLSLUniformBufferDeclaration(const TCHAR* Name,const FShade
 			TEXT("#endif\r\n"),
 			Name,
 			Name,
+			UE::ShaderParameters::kUniformBufferConstantBufferPrefix,
 			Name,
 			*Decl.ConstantBufferMembers,
 			*Decl.ResourceMembers,
diff --git a/Engine/Source/Runtime/RenderCore/Public/ShaderParameters.h b/Engine/Source/Runtime/RenderCore/Public/ShaderParameters.h
--- a/Engine/Source/Runtime/RenderCore/Public/ShaderParameters.h
+++ b/Engine/Source/Runtime/RenderCore/Public/ShaderParameters.h
@@ -14,6 +14,11 @@ class FShaderParameterMap;
 class FShaderParametersMetadata;
 struct FShaderCompilerEnvironment;
 
+namespace UE::ShaderParameters
+{
+	static constexpr const TCHAR* kUniformBufferConstantBufferPrefix = TEXT("UniformBufferConstants_");
+}
+
 RENDERCORE_API void CacheUniformBufferIncludes(TMap<const TCHAR*, struct FCachedUniformBufferDeclaration>& Cache, EShaderPlatform Platform);
 
 
