Migrating back to using UniformBuffer structs.
Since the cbuffer uses the same name as the struct, we have to rename the cbuffer using a consistent prefix. The reflection handling also needs to reinterpret the renamed cbuffer as the correct UniformBuffer.

diff --git a/Engine/Shaders/Public/Platform.ush b/Engine/Shaders/Public/Platform.ush
--- a/Engine/Shaders/Public/Platform.ush
+++ b/Engine/Shaders/Public/Platform.ush
@@ -200,6 +200,15 @@
 #define COMPILER_SUPPORTS_HLSL2021 0
 #endif
 
+#ifndef PLATFORM_SUPPORTS_UB_STRUCT
+#define PLATFORM_SUPPORTS_UB_STRUCT 0
+#endif
+
+#if PLATFORM_SUPPORTS_UB_STRUCT
+	#define UB_CB_NAME(X) UniformBufferConstants_##X
+#else
+	#define UB_CB_NAME(X) X
+#endif
 
 // ---------------------------------------------------- Alternative floating point types
 
diff --git a/Engine/Shaders/Public/ShaderVersion.ush b/Engine/Shaders/Public/ShaderVersion.ush
--- a/Engine/Shaders/Public/ShaderVersion.ush
+++ b/Engine/Shaders/Public/ShaderVersion.ush
@@ -2,4 +2,4 @@
 // Each time the console command is executed it generates a new GUID. As this file is included
 // in Platform.ush (which should be included in any shader) it allows to invalidate the shader DDC.
 // 
-// GUID = B543ABA7-900F-42E1-8443-9F39A584DC92
+// GUID = C8AC3233-E3DE-4604-A83A-CB41DD285817
diff --git a/Engine/Source/Developer/ShaderCompilerCommon/Private/ShaderCompilerCommon.cpp b/Engine/Source/Developer/ShaderCompilerCommon/Private/ShaderCompilerCommon.cpp
--- a/Engine/Source/Developer/ShaderCompilerCommon/Private/ShaderCompilerCommon.cpp
+++ b/Engine/Source/Developer/ShaderCompilerCommon/Private/ShaderCompilerCommon.cpp
@@ -469,6 +469,13 @@ TCHAR* FindNextUniformBufferReference(TCHAR* SearchPtr, const TCHAR* SearchStrin
 	return nullptr;
 }
 
+FString UE::ShaderCompilerCommon::RemoveConstantBufferPrefix(const FString& InName)
+{
+	FString NewName(InName);
+	NewName.RemoveFromStart(UE::ShaderCompilerCommon::kUniformBufferConstantBufferPrefix);
+	return NewName;
+}
+
 void HandleReflectedGlobalConstantBufferMember(
 	const FString& MemberName,
 	uint32 ConstantBufferIndex,
@@ -518,8 +525,10 @@ void HandleReflectedUniformBuffer(
 	FShaderCompilerOutput& CompilerOutput
 )
 {
+	FString AdjustedUniformBufferName(UE::ShaderCompilerCommon::RemoveConstantBufferPrefix(UniformBufferName));
+
 	CompilerOutput.ParameterMap.AddParameterAllocation(
-		*UniformBufferName,
+		*AdjustedUniformBufferName,
 		ReflectionSlot,
 		BaseIndex,
 		BufferSize,
diff --git a/Engine/Source/Developer/ShaderCompilerCommon/Public/ShaderCompilerCommon.h b/Engine/Source/Developer/ShaderCompilerCommon/Public/ShaderCompilerCommon.h
--- a/Engine/Source/Developer/ShaderCompilerCommon/Public/ShaderCompilerCommon.h
+++ b/Engine/Source/Developer/ShaderCompilerCommon/Public/ShaderCompilerCommon.h
@@ -8,6 +8,11 @@
 #include "CrossCompilerDefinitions.h"
 #include "ShaderConductorContext.h"
 
+namespace UE::ShaderCompilerCommon
+{
+	static constexpr const TCHAR* kUniformBufferConstantBufferPrefix = TEXT("UniformBufferConstants_");
+}
+
 /**
  * This function looks for resources specified in ResourceTableMap in the 
  * parameter map, adds them to the resource table, and removes them from the
@@ -139,6 +144,11 @@ private:
 	bool bMovedLoosedParametersToRootConstantBuffer = false;
 };
 
+namespace UE::ShaderCompilerCommon
+{
+	extern SHADERCOMPILERCOMMON_API FString RemoveConstantBufferPrefix(const FString& InName);
+}
+
 extern SHADERCOMPILERCOMMON_API void HandleReflectedGlobalConstantBufferMember(
 	const FString& MemberName,
 	uint32 ConstantBufferIndex,
diff --git a/Engine/Source/Developer/Windows/ShaderFormatD3D/Private/D3DShaderCompiler.cpp b/Engine/Source/Developer/Windows/ShaderFormatD3D/Private/D3DShaderCompiler.cpp
--- a/Engine/Source/Developer/Windows/ShaderFormatD3D/Private/D3DShaderCompiler.cpp
+++ b/Engine/Source/Developer/Windows/ShaderFormatD3D/Private/D3DShaderCompiler.cpp
@@ -1280,7 +1280,11 @@ void CompileD3DShader(const FShaderCompilerInput& Input, FShaderCompilerOutput&
 		return;
 	}
 
-	RemoveUniformBuffersFromSource(Input.Environment, PreprocessedShaderSource);
+	// Only use UniformBuffer structs on SM6 until we can fully vet SM5
+	if (Language != ELanguage::SM6)
+	{
+		RemoveUniformBuffersFromSource(Input.Environment, PreprocessedShaderSource);
+	}
 
 	// Process TEXT macro.
 	TransformStringIntoCharacterArray(PreprocessedShaderSource);
diff --git a/Engine/Source/Developer/Windows/ShaderFormatD3D/Private/D3DShaderCompiler.inl b/Engine/Source/Developer/Windows/ShaderFormatD3D/Private/D3DShaderCompiler.inl
--- a/Engine/Source/Developer/Windows/ShaderFormatD3D/Private/D3DShaderCompiler.inl
+++ b/Engine/Source/Developer/Windows/ShaderFormatD3D/Private/D3DShaderCompiler.inl
@@ -412,7 +412,7 @@ template <typename ID3D1xShaderReflection, typename D3D1x_SHADER_DESC, typename
 				{
 					UniformBufferNames.AddDefaulted(CBIndex - UniformBufferNames.Num() + 1);
 				}
-				UniformBufferNames[CBIndex] = UniformBufferName;
+				UniformBufferNames[CBIndex] = UE::ShaderCompilerCommon::RemoveConstantBufferPrefix(UniformBufferName);
 				UsedUniformBufferSlots[CBIndex] = true;
 			}
 
diff --git a/Engine/Source/Developer/Windows/ShaderFormatD3D/Private/ShaderFormatD3D.cpp b/Engine/Source/Developer/Windows/ShaderFormatD3D/Private/ShaderFormatD3D.cpp
--- a/Engine/Source/Developer/Windows/ShaderFormatD3D/Private/ShaderFormatD3D.cpp
+++ b/Engine/Source/Developer/Windows/ShaderFormatD3D/Private/ShaderFormatD3D.cpp
@@ -120,6 +120,7 @@ public:
 		{
 			Input.Environment.SetDefine(TEXT("SM6_PROFILE"), 1);
 			Input.Environment.SetDefine(TEXT("COMPILER_DXC"), 1);
+			Input.Environment.SetDefine(TEXT("PLATFORM_SUPPORTS_UB_STRUCT"), 1);
 
 			if (USE_SHADER_MODEL_6_6 && !Input.IsRayTracingShader())
 			{
diff --git a/Engine/Source/Runtime/RenderCore/Private/ShaderParameters.cpp b/Engine/Source/Runtime/RenderCore/Private/ShaderParameters.cpp
--- a/Engine/Source/Runtime/RenderCore/Private/ShaderParameters.cpp
+++ b/Engine/Source/Runtime/RenderCore/Private/ShaderParameters.cpp
@@ -278,7 +278,7 @@ static FString CreateHLSLUniformBufferDeclaration(const TCHAR* Name,const FShade
 		return FString::Printf(
 			TEXT("#ifndef __UniformBuffer_%s_Definition__\r\n")
 			TEXT("#define __UniformBuffer_%s_Definition__\r\n")
-			TEXT("cbuffer %s\r\n")
+			TEXT("cbuffer UB_CB_NAME(%s)\r\n")
 			TEXT("{\r\n")
 			TEXT("%s")
 			TEXT("}\r\n")
