Fix FHlslParser to support global scope.

diff --git a/Engine/Source/Developer/ShaderCompilerCommon/Private/HlslExpressionParser.inl b/Engine/Source/Developer/ShaderCompilerCommon/Private/HlslExpressionParser.inl
--- a/Engine/Source/Developer/ShaderCompilerCommon/Private/HlslExpressionParser.inl
+++ b/Engine/Source/Developer/ShaderCompilerCommon/Private/HlslExpressionParser.inl
@@ -824,6 +824,23 @@ namespace CrossCompiler
 			return ParseResultError();
 		}
 
+		auto ParseScopedIdentifier = [](FHlslScanner& Scanner, FLinearAllocator* Allocator, FString& Name, const FHlslToken* Token, const FHlslToken* Token1, const FHlslToken* Token2)
+		{
+			while (Token1 && Token1->Token == EHlslToken::ColonColon && Token2 && Token2->Token == EHlslToken::Identifier)
+			{
+				Name += Token1->String;
+				Name += Token2->String;
+				Scanner.Advance();
+				Scanner.Advance();
+				Token1 = Scanner.PeekToken();
+				Token2 = Scanner.PeekToken(1);
+			}
+
+			AST::FExpression* AtomExpression = new(Allocator) AST::FExpression(Allocator, AST::EOperators::Identifier, Token->SourceInfo);
+			AtomExpression->Identifier = Allocator->Strdup(*Name);
+			return AtomExpression;
+		};
+
 		AST::FExpression* AtomExpression = nullptr;
 		switch (Token->Token)
 		{
@@ -834,6 +851,22 @@ namespace CrossCompiler
 			AtomExpression->LiteralType = Token->LiteralType;
 			break;
 
+		case EHlslToken::ColonColon:
+		{
+			const FHlslToken* Token1 = Scanner.PeekToken(1);
+			if (Token1 && Token1->Token == EHlslToken::Identifier)
+			{
+				FString Name;
+				AtomExpression = ParseScopedIdentifier(Scanner, Allocator, Name, Token, Token, Token1);
+			}
+			else
+			{
+				Scanner.SourceError(TEXT("Expected identifier after ::"));
+				return ParseResultError();
+			}
+		}
+			break;
+
 		case EHlslToken::Identifier:
 		{
 			auto* Token1 = Scanner.PeekToken(1);
@@ -842,18 +875,7 @@ namespace CrossCompiler
 			{
 				FString Name = Token->String;
 				Scanner.Advance();
-				while (Token1 && Token1->Token == EHlslToken::ColonColon && Token2 && Token2->Token == EHlslToken::Identifier)
-				{
-					Name += Token1->String;
-					Name += Token2->String;
-					Scanner.Advance();
-					Scanner.Advance();
-					Token1 = Scanner.PeekToken();
-					Token2 = Scanner.PeekToken(1);
-				}
-
-				AtomExpression = new(Allocator) AST::FExpression(Allocator, AST::EOperators::Identifier, Token->SourceInfo);
-				AtomExpression->Identifier = Allocator->Strdup(*Name);
+				AtomExpression = ParseScopedIdentifier(Scanner, Allocator, Name, Token, Token1, Token2);
 			}
 			else
 			{
