Bindless shader parameters
- Adding support for bindless resources/samplers. The loose resources have to be wrapped in a macro that generates an index and proxy resource when bindless is enabled. Reflection has to parse these differently from other resources since they're parsed as uint parameters. The runtime also has to bind these differently since they end up needing to find their real resource in the binding data and then update a constant buffer.
- To assist in conversions, the shader compiler will detect if a shader is compiled with both bindless and non-bindless resources/samplers and will emit errors if that is the case.
- Adding support for bindless Uniform Buffer resources/samplers. Because Uniform Buffers aren't per-shader, they'll always update their constants with the bindless indices of their resources.
- Adding more flags to FShaderCodePackedResourceCounts required changing all the array initializations to individual parameters since not all the new flags are used everywhere.
- No shaders have been configured to support bindless resources/samplers. Yet.

diff --git a/Engine/Shaders/Public/BindlessResources.ush b/Engine/Shaders/Public/BindlessResources.ush
new file mode 100644
--- /dev/null
+++ b/Engine/Shaders/Public/BindlessResources.ush
@@ -0,0 +1,91 @@
+// Copyright Epic Games, Inc. All Rights Reserved.
+
+#pragma once
+
+#include "Platform.ush"
+
+#ifndef ENABLE_BINDLESS_RESOURCES
+#define ENABLE_BINDLESS_RESOURCES 0
+#endif
+
+#ifndef ENABLE_BINDLESS_SAMPLERS
+#define ENABLE_BINDLESS_SAMPLERS 0
+#endif
+
+// Indices can always be used as-is. This is important for things like Uniform Buffer layouts.
+
+#define DEFINE_RESOURCE_INDEX(Name)  uint BindlessResource_##Name
+#define DEFINE_SAMPLER_INDEX(Name)   uint BindlessSampler_##Name
+
+//
+
+#if ENABLE_BINDLESS_RESOURCES && PLATFORM_SUPPORTS_BINDLESS
+    #define DEFINE_RESOURCE_VARIABLE(Type, Name) \
+        Type GetBindlessResource##Name() \
+        { \
+            return GetResourceFromHeap(Type, BindlessResource_##Name); \
+        }
+
+    #define GET_RESOURCE(Name) GetBindlessResource##Name()
+
+    #define DEFINE_RESOURCE(Type, Name) \
+		DEFINE_RESOURCE_INDEX(Name); \
+        DEFINE_RESOURCE_VARIABLE(Type, Name); \
+		static const Type Name = GET_RESOURCE(Name);
+#else
+	#define DEFINE_RESOURCE_VARIABLE(Type, Name)  Type Name
+	#define DEFINE_RESOURCE(Type, Name)           Type Name
+	#define GET_RESOURCE(Name)                    Name
+#endif
+
+#if ENABLE_BINDLESS_SAMPLERS && PLATFORM_SUPPORTS_BINDLESS
+	#define DEFINE_SAMPLER_VARIABLE(Type, Name) \
+		Type GetBindlessSampler##Name() \
+        { \
+            return GetSamplerFromHeap(Type, BindlessSampler_##Name); \
+        }
+
+    #define GET_SAMPLER(Name)  GetBindlessSampler##Name()
+
+    #define DEFINE_SAMPLER(Type, Name) \
+        DEFINE_SAMPLER_INDEX(Name); \
+        DEFINE_SAMPLER_VARIABLE(Type, Name); \
+		static const Type Name = GET_SAMPLER(Name);
+#else
+	#define DEFINE_SAMPLER_VARIABLE(Type, Name)  Type Name
+	#define DEFINE_SAMPLER(Type, Name)           Type Name
+	#define GET_SAMPLER(Name)                    Name
+#endif
+
+/*
+// Uniform buffers will do:
+cbuffer UniformBuffer_MyUniformBuffer
+{
+	DEFINE_RESOURCE_INDEX(MyUniformBuffer_MyTexture);
+	DEFINE_SAMPLER_INDEX(MyUniformBuffer_MySampler);
+};
+
+DEFINE_RESOURCE_VARIABLE(Texture2D, MyUniformBuffer_MyTexture);
+DEFINE_SAMPLER_VARIABLE(SamplerState, MyUniformBuffer_MySampler);
+
+static const struct
+{
+	Texture2D MyTexture;
+	SamplerState MySampler;
+}
+MyUniformBuffer = { GET_RESOURCE(MyUniformBuffer_MyTexture), GET_SAMPLER(MyUniformBuffer_MySampler) };
+
+{	// No wrappers necessary to use
+	MyUniformBuffer.MyTexture.Sample(MyUniformBuffer.MySampler, ...);
+}
+
+
+// Normal uses will look like:
+DEFINE_RESOURCE(Texture2D, MyTexture);
+DEFINE_SAMPLER(SamplerState, MySampler);
+
+{	// No wrappers necessary to use
+	MyTexture.Sample(MySampler, ...);
+}
+
+*/
diff --git a/Engine/Shaders/Public/Platform.ush b/Engine/Shaders/Public/Platform.ush
--- a/Engine/Shaders/Public/Platform.ush
+++ b/Engine/Shaders/Public/Platform.ush
@@ -36,6 +36,7 @@
 	#include "Platform/D3D/D3DCommon.ush"
 #endif
 
+#include "/Engine/Public/BindlessResources.ush"
 
 // ---------------------------------------------------- DDC invalidation
 // to support the console command "r.InvalidateShaderCache"
diff --git a/Engine/Shaders/Public/ShaderVersion.ush b/Engine/Shaders/Public/ShaderVersion.ush
--- a/Engine/Shaders/Public/ShaderVersion.ush
+++ b/Engine/Shaders/Public/ShaderVersion.ush
@@ -2,4 +2,4 @@
 // Each time the console command is executed it generates a new GUID. As this file is included
 // in Platform.ush (which should be included in any shader) it allows to invalidate the shader DDC.
 // 
-// GUID = 0E1525E8-83A2-4FB1-BB2E-8BABCE046416
+// GUID = FF4213C8-41E0-41E2-A4D2-A09146AC6BB6
diff --git a/Engine/Source/Developer/ShaderCompilerCommon/Private/ShaderCompilerCommon.cpp b/Engine/Source/Developer/ShaderCompilerCommon/Private/ShaderCompilerCommon.cpp
--- a/Engine/Source/Developer/ShaderCompilerCommon/Private/ShaderCompilerCommon.cpp
+++ b/Engine/Source/Developer/ShaderCompilerCommon/Private/ShaderCompilerCommon.cpp
@@ -476,32 +476,68 @@ FString UE::ShaderCompilerCommon::RemoveConstantBufferPrefix(const FString& InNa
 	return NewName;
 }
 
+EShaderParameterType UE::ShaderCompilerCommon::ParseAndRemoveBindlessParameterPrefix(FString& InName)
+{
+	if (InName.RemoveFromStart(UE::ShaderCompilerCommon::kBindlessResourcePrefix))
+	{
+		return EShaderParameterType::BindlessResourceIndex;
+	}
+
+	if (InName.RemoveFromStart(UE::ShaderCompilerCommon::kBindlessSamplerPrefix))
+	{
+		return EShaderParameterType::BindlessSamplerIndex;
+	}
+
+	return EShaderParameterType::LooseData;
+}
+
+bool UE::ShaderCompilerCommon::RemoveBindlessParameterPrefix(FString& InName)
+{
+	return ParseAndRemoveBindlessParameterPrefix(InName) != EShaderParameterType::LooseData;
+}
+
 void HandleReflectedGlobalConstantBufferMember(
-	const FString& MemberName,
+	const FString& InMemberName,
 	uint32 ConstantBufferIndex,
 	int32 ReflectionOffset,
 	int32 ReflectionSize,
 	FShaderCompilerOutput& Output
 )
 {
+	FString MemberName = InMemberName;
+	const EShaderParameterType ParameterType = UE::ShaderCompilerCommon::ParseAndRemoveBindlessParameterPrefix(MemberName);
+
 	Output.ParameterMap.AddParameterAllocation(
 		*MemberName,
 		ConstantBufferIndex,
 		ReflectionOffset,
 		ReflectionSize,
-		EShaderParameterType::LooseData);
+		ParameterType);
 }
 
 void HandleReflectedRootConstantBufferMember(
 	const FShaderCompilerInput& Input,
 	const FShaderParameterParser& ShaderParameterParser,
-	const FString& MemberName,
+	const FString& InMemberName,
 	int32 ReflectionOffset,
 	int32 ReflectionSize,
 	FShaderCompilerOutput& Output
 )
 {
-	ShaderParameterParser.ValidateShaderParameterType(Input, MemberName, ReflectionOffset, ReflectionSize, Output);
+	ShaderParameterParser.ValidateShaderParameterType(Input, InMemberName, ReflectionOffset, ReflectionSize, Output);
+
+	FString MemberName = InMemberName;
+	const EShaderParameterType ParameterType = UE::ShaderCompilerCommon::ParseAndRemoveBindlessParameterPrefix(MemberName);
+
+	if (ParameterType != EShaderParameterType::LooseData)
+	{
+		Output.ParameterMap.AddParameterAllocation(
+			*MemberName,
+			FShaderParametersMetadata::kRootCBufferBindingIndex,
+			ReflectionOffset,
+			1,
+			ParameterType);
+	}
 }
 
 void HandleReflectedRootConstantBuffer(
@@ -653,7 +689,12 @@ void AddUnboundShaderParameterError(
 inline bool MemberWasPotentiallyMoved(const FShaderParametersMetadata::FMember& InMember)
 {
 	const EUniformBufferBaseType BaseType = InMember.GetBaseType();
-	return BaseType == UBMT_INT32 || BaseType == UBMT_UINT32 || BaseType == UBMT_FLOAT32;
+	return BaseType == UBMT_INT32 || BaseType == UBMT_UINT32 || BaseType == UBMT_FLOAT32
+		|| BaseType == UBMT_TEXTURE || BaseType == UBMT_SRV || BaseType == UBMT_UAV
+		|| BaseType == UBMT_RDG_TEXTURE || BaseType == UBMT_RDG_TEXTURE_SRV || BaseType == UBMT_RDG_TEXTURE_UAV
+		|| BaseType == UBMT_RDG_BUFFER_SRV || BaseType == UBMT_RDG_BUFFER_UAV
+		|| BaseType == UBMT_SAMPLER
+		;
 }
 
 bool FShaderParameterParser::ParseAndMoveShaderParametersToRootConstantBuffer(
@@ -774,6 +815,9 @@ bool FShaderParameterParser::ParseAndMoveShaderParametersToRootConstantBuffer(
 
 				FString Type = PreprocessedShaderSource.Mid(TypeStartPos, TypeEndPos - TypeStartPos + 1);
 				FString Name = PreprocessedShaderSource.Mid(NameStartPos, NameEndPos - NameStartPos + 1);
+				const FString ParsedName = Name;
+
+				const bool bBindlessParameter = UE::ShaderCompilerCommon::RemoveBindlessParameterPrefix(Name);
 
 				FParsedShaderParameter ParsedParameter;
 				bool bUpdateParsedParameters = false;
@@ -794,7 +838,7 @@ bool FShaderParameterParser::ParseAndMoveShaderParametersToRootConstantBuffer(
 						if (bMoveToRootConstantBuffer && ParsedParameter.IsBindable())
 						{
 							EUniformBufferBaseType BaseType = ParsedParameter.Member->GetBaseType();
-							bEraseOriginalParameter = BaseType == UBMT_INT32 || BaseType == UBMT_UINT32 || BaseType == UBMT_FLOAT32;
+							bEraseOriginalParameter = BaseType == UBMT_INT32 || BaseType == UBMT_UINT32 || BaseType == UBMT_FLOAT32 || bBindlessParameter;
 						}
 					}
 				}
@@ -807,9 +851,11 @@ bool FShaderParameterParser::ParseAndMoveShaderParametersToRootConstantBuffer(
 				// Update 
 				if (bUpdateParsedParameters)
 				{
+					ParsedParameter.ParsedName = ParsedName;
 					ParsedParameter.ParsedType = Type;
 					ParsedParameter.ParsedPragmaLineoffset = CurrentPragamLineoffset;
 					ParsedParameter.ParsedLineOffset = CurrentLineoffset;
+					ParsedParameter.bOriginalParameterErased = bEraseOriginalParameter;
 
 					if (ArrayStartPos != -1 && ArrayEndPos != -1)
 					{
@@ -1137,7 +1183,7 @@ bool FShaderParameterParser::ParseAndMoveShaderParametersToRootConstantBuffer(
 			if (MemberWasPotentiallyMoved(Member))
 			{
 				FParsedShaderParameter* ParsedParameter = ParsedParameters.Find(ShaderBindingName);
-				if (ParsedParameter && ParsedParameter->IsFound())
+				if (ParsedParameter && ParsedParameter->IsFound() && ParsedParameter->bOriginalParameterErased)
 				{
 					uint32 ConstantRegister = ByteOffset / 16;
 					const TCHAR* ConstantSwizzle = [ByteOffset]()
@@ -1157,7 +1203,7 @@ bool FShaderParameterParser::ParseAndMoveShaderParametersToRootConstantBuffer(
 						RootCBufferContent.Append(FString::Printf(
 							TEXT("%s %s[%s] : packoffset(c%d%s);\r\n"),
 							*ParsedParameter->ParsedType,
-							ShaderBindingName,
+							*ParsedParameter->ParsedName,
 							*ParsedParameter->ParsedArraySize,
 							ConstantRegister,
 							ConstantSwizzle));
@@ -1167,7 +1213,7 @@ bool FShaderParameterParser::ParseAndMoveShaderParametersToRootConstantBuffer(
 						RootCBufferContent.Append(FString::Printf(
 							TEXT("%s %s : packoffset(c%d%s);\r\n"),
 							*ParsedParameter->ParsedType,
-							ShaderBindingName,
+							*ParsedParameter->ParsedName,
 							ConstantRegister,
 							ConstantSwizzle));
 					}
@@ -1204,7 +1250,11 @@ void FShaderParameterParser::ValidateShaderParameterType(
 	bool bPlatformSupportsPrecisionModifier,
 	FShaderCompilerOutput& CompilerOutput) const
 {
-	const FShaderParameterParser::FParsedShaderParameter& ParsedParameter = FindParameterInfos(ShaderBindingName);
+	FString BindingName(ShaderBindingName);
+
+	const bool bBindlessHack = UE::ShaderCompilerCommon::RemoveBindlessParameterPrefix(BindingName);
+
+	const FShaderParameterParser::FParsedShaderParameter& ParsedParameter = FindParameterInfos(BindingName);
 
 	check(ParsedParameter.IsFound());
 	check(CompilerInput.RootParametersStructure);
@@ -1216,6 +1266,7 @@ void FShaderParameterParser::ValidateShaderParameterType(
 	}
 
 	// Validate the shader type.
+	if (!bBindlessHack)
 	{
 		FString ExpectedShaderType;
 		ParsedParameter.Member->GenerateShaderParameterType(ExpectedShaderType, bPlatformSupportsPrecisionModifier);
@@ -1289,7 +1340,7 @@ void FShaderParameterParser::ValidateShaderParameterType(
 	}
 
 	// Validate parameter size, in case this is an array.
-	if (ReflectionSize > int32(ParsedParameter.Member->GetMemberSize()))
+	if (!bBindlessHack && ReflectionSize > int32(ParsedParameter.Member->GetMemberSize()))
 	{
 		FString CppCodeName = CompilerInput.RootParametersStructure->GetFullMemberCodeName(ParsedParameter.ConstantBufferOffset);
 
diff --git a/Engine/Source/Developer/ShaderCompilerCommon/Public/ShaderCompilerCommon.h b/Engine/Source/Developer/ShaderCompilerCommon/Public/ShaderCompilerCommon.h
--- a/Engine/Source/Developer/ShaderCompilerCommon/Public/ShaderCompilerCommon.h
+++ b/Engine/Source/Developer/ShaderCompilerCommon/Public/ShaderCompilerCommon.h
@@ -11,6 +11,8 @@
 namespace UE::ShaderCompilerCommon
 {
 	static constexpr const TCHAR* kUniformBufferConstantBufferPrefix = TEXT("UniformBufferConstants_");
+	static constexpr const TCHAR* kBindlessResourcePrefix = TEXT("BindlessResource_");
+	static constexpr const TCHAR* kBindlessSamplerPrefix = TEXT("BindlessSampler_");
 }
 
 /**
@@ -58,12 +60,16 @@ public:
 		const FShaderParametersMetadata::FMember* Member = nullptr;
 
 		/** Information found about the member when parsing the preprocessed code. */
+		FString ParsedName;
 		FString ParsedType;
 		FString ParsedArraySize;
 
 		/** Offset the member should be in the constant buffer. */
 		int32 ConstantBufferOffset = 0;
 
+		/** Records if the parameter was fully moved. */
+		bool bOriginalParameterErased{};
+
 		/* Returns whether the shader parameter has been found when parsing. */
 		bool IsFound() const
 		{
@@ -146,7 +152,9 @@ private:
 
 namespace UE::ShaderCompilerCommon
 {
-	extern SHADERCOMPILERCOMMON_API FString RemoveConstantBufferPrefix(const FString& InName);
+	extern SHADERCOMPILERCOMMON_API FString              RemoveConstantBufferPrefix(const FString& InName);
+	extern SHADERCOMPILERCOMMON_API EShaderParameterType ParseAndRemoveBindlessParameterPrefix(FString& InName);
+	extern SHADERCOMPILERCOMMON_API bool                 RemoveBindlessParameterPrefix(FString& InName);
 }
 
 extern SHADERCOMPILERCOMMON_API void HandleReflectedGlobalConstantBufferMember(
diff --git a/Engine/Source/Developer/Windows/ShaderFormatD3D/Private/D3DShaderCompiler.cpp b/Engine/Source/Developer/Windows/ShaderFormatD3D/Private/D3DShaderCompiler.cpp
--- a/Engine/Source/Developer/Windows/ShaderFormatD3D/Private/D3DShaderCompiler.cpp
+++ b/Engine/Source/Developer/Windows/ShaderFormatD3D/Private/D3DShaderCompiler.cpp
@@ -1136,7 +1136,17 @@ bool CompileAndProcessD3DShaderFXC(FString& PreprocessedShaderSource, const FStr
 				Output.ShaderCode.AddOptionalData(ResourceMasks);
 			};
 
-			FShaderCodePackedResourceCounts PackedResourceCounts = { bGlobalUniformBufferUsed, static_cast<uint8>(NumSamplers), static_cast<uint8>(NumSRVs), static_cast<uint8>(NumCBs), static_cast<uint8>(NumUAVs) };
+			FShaderCodePackedResourceCounts PackedResourceCounts{};
+			if (bGlobalUniformBufferUsed)
+			{
+				PackedResourceCounts.UsageFlags |= EShaderResourceUsageFlags::GlobalUniformBuffer;
+			}
+
+			PackedResourceCounts.NumSamplers = static_cast<uint8>(NumSamplers);
+			PackedResourceCounts.NumSRVs = static_cast<uint8>(NumSRVs);
+			PackedResourceCounts.NumCBs = static_cast<uint8>(NumCBs);
+			PackedResourceCounts.NumUAVs = static_cast<uint8>(NumUAVs);
+
 			GenerateFinalOutput(CompressedData,
 				Input, VendorExtensions,
 				UsedUniformBufferSlots, UniformBufferNames,
diff --git a/Engine/Source/Developer/Windows/ShaderFormatD3D/Private/D3DShaderCompilerDXC.cpp b/Engine/Source/Developer/Windows/ShaderFormatD3D/Private/D3DShaderCompilerDXC.cpp
--- a/Engine/Source/Developer/Windows/ShaderFormatD3D/Private/D3DShaderCompilerDXC.cpp
+++ b/Engine/Source/Developer/Windows/ShaderFormatD3D/Private/D3DShaderCompilerDXC.cpp
@@ -885,6 +885,43 @@ bool CompileAndProcessD3DShaderDXC(FString& PreprocessedShaderSource,
 			Output.bSucceeded = false;
 		}
 
+		const bool bBindlessResourcesUsed = ((ShaderRequiresFlags & D3D_SHADER_REQUIRES_RESOURCE_DESCRIPTOR_HEAP_INDEXING) != 0);
+		const bool bBindlessSamplersUsed = ((ShaderRequiresFlags & D3D_SHADER_REQUIRES_SAMPLER_DESCRIPTOR_HEAP_INDEXING) != 0);
+
+		if (Output.bSucceeded)
+		{
+			auto GetAllResourcesOfType = [&](EShaderParameterType InType)
+			{
+				const TArray<FString> AllNames = Output.ParameterMap.GetAllParameterNamesOfType(InType);
+				if (AllNames.IsEmpty())
+				{
+					return FString();
+				}
+				return FString::Join(AllNames, TEXT(", "));
+			};
+
+			if (bBindlessResourcesUsed && NumSRVs > 0)
+			{
+				const FString Names = GetAllResourcesOfType(EShaderParameterType::SRV);
+				FilteredErrors.Add(FString::Printf(TEXT("Shader is mixing bindless resources with non-bindless resources. %d SRV slots were detected: %s"), NumSRVs, *Names));
+				Output.bSucceeded = false;
+			}
+
+			if (bBindlessResourcesUsed && NumUAVs > 0)
+			{
+				const FString Names = GetAllResourcesOfType(EShaderParameterType::UAV);
+				FilteredErrors.Add(FString::Printf(TEXT("Shader is mixing bindless resources with non-bindless resources. %d UAV slots were detected: %s"), NumUAVs, *Names));
+				Output.bSucceeded = false;
+			}
+
+			if (bBindlessSamplersUsed && NumSamplers > 0)
+			{
+				const FString Names = GetAllResourcesOfType(EShaderParameterType::Sampler);
+				FilteredErrors.Add(FString::Printf(TEXT("Shader is mixing bindless samplers with non-bindless samplers. %d sampler slots were detected: %s"), NumSamplers, *Names));
+				Output.bSucceeded = false;
+			}
+		}
+
 		// Save results if compilation and reflection succeeded
 		if (Output.bSucceeded)
 		{
@@ -949,7 +986,25 @@ bool CompileAndProcessD3DShaderDXC(FString& PreprocessedShaderSource,
 
 			//#todo-rco: Should compress ShaderCode?
 
-			FShaderCodePackedResourceCounts PackedResourceCounts = { bGlobalUniformBufferUsed, static_cast<uint8>(NumSamplers), static_cast<uint8>(NumSRVs), static_cast<uint8>(NumCBs), static_cast<uint8>(NumUAVs) };
+			FShaderCodePackedResourceCounts PackedResourceCounts{};
+			if (bGlobalUniformBufferUsed)
+			{
+				PackedResourceCounts.UsageFlags |= EShaderResourceUsageFlags::GlobalUniformBuffer;
+			}
+			if (bBindlessResourcesUsed)
+			{
+				PackedResourceCounts.UsageFlags |= EShaderResourceUsageFlags::BindlessResources;
+			}
+			if (bBindlessSamplersUsed)
+			{
+				PackedResourceCounts.UsageFlags |= EShaderResourceUsageFlags::BindlessSamplers;
+			}
+
+			PackedResourceCounts.NumSamplers = static_cast<uint8>(NumSamplers);
+			PackedResourceCounts.NumSRVs = static_cast<uint8>(NumSRVs);
+			PackedResourceCounts.NumCBs = static_cast<uint8>(NumCBs);
+			PackedResourceCounts.NumUAVs = static_cast<uint8>(NumUAVs);
+
 			GenerateFinalOutput(ShaderBlob,
 				Input, VendorExtensions,
 				UsedUniformBufferSlots, UniformBufferNames,
diff --git a/Engine/Source/Runtime/D3D12RHI/Private/D3D12Commands.cpp b/Engine/Source/Runtime/D3D12RHI/Private/D3D12Commands.cpp
--- a/Engine/Source/Runtime/D3D12RHI/Private/D3D12Commands.cpp
+++ b/Engine/Source/Runtime/D3D12RHI/Private/D3D12Commands.cpp
@@ -121,7 +121,7 @@ void FD3D12CommandContext::RHIDispatchComputeShader(uint32 ThreadGroupCountX, ui
 		GetParentDevice()->RegisterGPUDispatch(FIntVector(ThreadGroupCountX, ThreadGroupCountY, ThreadGroupCountZ));	
 	}
 
-	if (ComputeShader->ResourceCounts.bGlobalUniformBufferUsed)
+	if (ComputeShader->UsesGlobalUniformBuffer())
 	{
 		CommitComputeShaderConstants();
 	}
@@ -148,7 +148,7 @@ void FD3D12CommandContext::RHIDispatchIndirectComputeShader(FRHIBuffer* Argument
 	FD3D12ComputeShader* ComputeShader = nullptr;
 	StateCache.GetComputeShader(&ComputeShader);
 
-	if (ComputeShader->ResourceCounts.bGlobalUniformBufferUsed)
+	if (ComputeShader->UsesGlobalUniformBuffer())
 	{
 		CommitComputeShaderConstants();
 	}
diff --git a/Engine/Source/Runtime/D3D12RHI/Private/D3D12Device.cpp b/Engine/Source/Runtime/D3D12RHI/Private/D3D12Device.cpp
--- a/Engine/Source/Runtime/D3D12RHI/Private/D3D12Device.cpp
+++ b/Engine/Source/Runtime/D3D12RHI/Private/D3D12Device.cpp
@@ -352,7 +352,7 @@ void FD3D12Device::SetupAfterDeviceCreation()
 		OfflineDescriptorManagers[Index].Init(static_cast<ERHIDescriptorHeapType>(Index));
 	}
 
-	GlobalSamplerHeap.Init(GGlobalSamplerHeapSize);
+	GlobalSamplerHeap.Init(GGlobalSamplerHeapSize - GBindlessSamplerDescriptorHeapSize);
 
 	OnlineDescriptorManager.Init(GOnlineDescriptorHeapSize, GOnlineDescriptorHeapBlockSize);
 
diff --git a/Engine/Source/Runtime/D3D12RHI/Private/D3D12Pipelinestate.cpp b/Engine/Source/Runtime/D3D12RHI/Private/D3D12Pipelinestate.cpp
--- a/Engine/Source/Runtime/D3D12RHI/Private/D3D12Pipelinestate.cpp
+++ b/Engine/Source/Runtime/D3D12RHI/Private/D3D12Pipelinestate.cpp
@@ -325,11 +325,11 @@ FD3D12GraphicsPipelineState::FD3D12GraphicsPipelineState(
 		FMemory::Memcpy(StreamStrides, ((FD3D12VertexDeclaration*) Initializer.BoundShaderState.VertexDeclarationRHI)->StreamStrides, sizeof(StreamStrides));
 	else
 		FMemory::Memzero(StreamStrides, sizeof(StreamStrides));
-	bShaderNeedsGlobalConstantBuffer[SF_Vertex] = GetVertexShader() && GetVertexShader()->ResourceCounts.bGlobalUniformBufferUsed;
-	bShaderNeedsGlobalConstantBuffer[SF_Mesh] = GetMeshShader() && GetMeshShader()->ResourceCounts.bGlobalUniformBufferUsed;
-	bShaderNeedsGlobalConstantBuffer[SF_Amplification] = GetAmplificationShader() && GetAmplificationShader()->ResourceCounts.bGlobalUniformBufferUsed;
-	bShaderNeedsGlobalConstantBuffer[SF_Pixel] = GetPixelShader() && GetPixelShader()->ResourceCounts.bGlobalUniformBufferUsed;
-	bShaderNeedsGlobalConstantBuffer[SF_Geometry] = GetGeometryShader() && GetGeometryShader()->ResourceCounts.bGlobalUniformBufferUsed;
+	bShaderNeedsGlobalConstantBuffer[SF_Vertex] = GetVertexShader() && GetVertexShader()->UsesGlobalUniformBuffer();
+	bShaderNeedsGlobalConstantBuffer[SF_Mesh] = GetMeshShader() && GetMeshShader()->UsesGlobalUniformBuffer();
+	bShaderNeedsGlobalConstantBuffer[SF_Amplification] = GetAmplificationShader() && GetAmplificationShader()->UsesGlobalUniformBuffer();
+	bShaderNeedsGlobalConstantBuffer[SF_Pixel] = GetPixelShader() && GetPixelShader()->UsesGlobalUniformBuffer();
+	bShaderNeedsGlobalConstantBuffer[SF_Geometry] = GetGeometryShader() && GetGeometryShader()->UsesGlobalUniformBuffer();
 
 	// GRHISupportsPipelineStateSortKey
 	SetSortKey(InPipelineState->GetContextSortKey());
diff --git a/Engine/Source/Runtime/D3D12RHI/Private/D3D12RayTracing.cpp b/Engine/Source/Runtime/D3D12RHI/Private/D3D12RayTracing.cpp
--- a/Engine/Source/Runtime/D3D12RHI/Private/D3D12RayTracing.cpp
+++ b/Engine/Source/Runtime/D3D12RHI/Private/D3D12RayTracing.cpp
@@ -2489,7 +2489,7 @@ public:
 		{
 			FD3D12RayTracingShader* Shader = FD3D12DynamicRHI::ResourceCast(ShaderRHI);
 			checkf(Shader->pRootSignature->GetRootSignature() == GlobalRootSignature, TEXT("All raygen and miss shaders must share the same root signature"));
-			checkf(!Shader->ResourceCounts.bGlobalUniformBufferUsed, TEXT("Global uniform buffers are not implemented for ray generation shaders"));
+			checkf(!Shader->UsesGlobalUniformBuffer(), TEXT("Global uniform buffers are not implemented for ray generation shaders"));
 
 			FD3D12RayTracingPipelineCache::FEntry* ShaderCacheEntry = AddShaderCollection(Shader, FD3D12RayTracingPipelineCache::ECollectionType::RayGen);
 
@@ -2513,7 +2513,7 @@ public:
 			FD3D12RayTracingShader* Shader = FD3D12DynamicRHI::ResourceCast(ShaderRHI);
 
 			checkf(Shader, TEXT("A valid ray tracing shader must be provided for all elements in the FRayTracingPipelineStateInitializer miss shader table."));
-			checkf(!Shader->ResourceCounts.bGlobalUniformBufferUsed, TEXT("Global uniform buffers are not implemented for ray tracing miss shaders"));
+			checkf(!Shader->UsesGlobalUniformBuffer(), TEXT("Global uniform buffers are not implemented for ray tracing miss shaders"));
 
 			const uint32 ShaderViewDescriptors = Shader->ResourceCounts.NumSRVs + Shader->ResourceCounts.NumUAVs;
 			MaxHitGroupViewDescriptors = FMath::Max(MaxHitGroupViewDescriptors, ShaderViewDescriptors);
@@ -2568,7 +2568,7 @@ public:
 			FD3D12RayTracingShader* Shader = FD3D12DynamicRHI::ResourceCast(ShaderRHI);
 
 			checkf(Shader, TEXT("A valid ray tracing shader must be provided for all elements in the FRayTracingPipelineStateInitializer callable shader table."));
-			checkf(!Shader->ResourceCounts.bGlobalUniformBufferUsed, TEXT("Global uniform buffers are not implemented for ray tracing callable shaders"));
+			checkf(!Shader->UsesGlobalUniformBuffer(), TEXT("Global uniform buffers are not implemented for ray tracing callable shaders"));
 
 			const uint32 ShaderViewDescriptors = Shader->ResourceCounts.NumSRVs + Shader->ResourceCounts.NumUAVs;
 			MaxHitGroupViewDescriptors = FMath::Max(MaxHitGroupViewDescriptors, ShaderViewDescriptors);
@@ -4981,12 +4981,12 @@ static bool SetRayTracingShaderResources(
 
 	// Bind loose parameters
 
-	if (Shader->ResourceCounts.bGlobalUniformBufferUsed)
+	if (Shader->UsesGlobalUniformBuffer())
 	{
 		checkf(InLooseParameterDataSize && InLooseParameterData, TEXT("Shader uses global uniform buffer, but the required loose parameter data is not provided."));
 	}
 
-	if (InLooseParameterData && Shader->ResourceCounts.bGlobalUniformBufferUsed)
+	if (InLooseParameterData && Shader->UsesGlobalUniformBuffer())
 	{
 		const uint32 CBVIndex = 0; // Global uniform buffer is always assumed to be in slot 0
 		D3D12_GPU_VIRTUAL_ADDRESS BufferAddress = Binder.CreateTransientConstantBuffer(InLooseParameterData, InLooseParameterDataSize);
diff --git a/Engine/Source/Runtime/D3D12RHI/Private/D3D12RootSignature.cpp b/Engine/Source/Runtime/D3D12RHI/Private/D3D12RootSignature.cpp
--- a/Engine/Source/Runtime/D3D12RHI/Private/D3D12RootSignature.cpp
+++ b/Engine/Source/Runtime/D3D12RHI/Private/D3D12RootSignature.cpp
@@ -309,7 +309,7 @@ FD3D12RootSignatureDesc::FD3D12RootSignatureDesc(const FD3D12QuantizedBoundShade
 	}
 	else if(QBSS.RootSignatureType == RS_RayTracingGlobal)
 	{
-		Flags = FD3D12_ROOT_SIGNATURE_FLAG_GLOBAL_ROOT_SIGNATURE;
+		Flags |= FD3D12_ROOT_SIGNATURE_FLAG_GLOBAL_ROOT_SIGNATURE;
 	}
 	else if (QBSS.RootSignatureType == RS_Raster)
 #endif // D3D12_RHI_RAYTRACING
diff --git a/Engine/Source/Runtime/D3D12RHI/Private/D3D12Shader.h b/Engine/Source/Runtime/D3D12RHI/Private/D3D12Shader.h
--- a/Engine/Source/Runtime/D3D12RHI/Private/D3D12Shader.h
+++ b/Engine/Source/Runtime/D3D12RHI/Private/D3D12Shader.h
@@ -91,8 +91,9 @@ struct FD3D12ShaderData
 #endif
 
 	FORCEINLINE bool UsesDiagnosticBuffer() const { return EnumHasAnyFlags(GetFeatures(), EShaderCodeFeatures::DiagnosticBuffer); }
-	FORCEINLINE bool UsesBindlessResources() const { return EnumHasAnyFlags(GetFeatures(), EShaderCodeFeatures::BindlessResources); }
-	FORCEINLINE bool UsesBindlessSamplers() const { return EnumHasAnyFlags(GetFeatures(), EShaderCodeFeatures::BindlessSamplers); }
+	FORCEINLINE bool UsesGlobalUniformBuffer() const { return EnumHasAnyFlags(ResourceCounts.UsageFlags, EShaderResourceUsageFlags::GlobalUniformBuffer); }
+	FORCEINLINE bool UsesBindlessResources() const { return EnumHasAnyFlags(ResourceCounts.UsageFlags, EShaderResourceUsageFlags::BindlessResources); }
+	FORCEINLINE bool UsesBindlessSamplers() const { return EnumHasAnyFlags(ResourceCounts.UsageFlags, EShaderResourceUsageFlags::BindlessSamplers); }
 
 	bool InitCommon(TArrayView<const uint8> InCode);
 };
diff --git a/Engine/Source/Runtime/D3D12RHI/Private/D3D12UniformBuffer.cpp b/Engine/Source/Runtime/D3D12RHI/Private/D3D12UniformBuffer.cpp
--- a/Engine/Source/Runtime/D3D12RHI/Private/D3D12UniformBuffer.cpp
+++ b/Engine/Source/Runtime/D3D12RHI/Private/D3D12UniformBuffer.cpp
@@ -8,6 +8,111 @@
 #include "UniformBuffer.h"
 #include "ShaderParameterStruct.h"
 
+/// Utility class for reading shader parameters out of the data blob passed in.
+struct FUniformDataReader
+{
+	FUniformDataReader() = delete;
+	FUniformDataReader(const void* InData) : Data(reinterpret_cast<const uint8*>(InData)) { }
+
+	template<typename TResourceOut>
+	const TResourceOut& Read(const FRHIUniformBufferResource& InResource) const
+	{
+		return *reinterpret_cast<const TResourceOut*>(Data + InResource.MemberOffset);
+	}
+
+	const uint8* Data;
+};
+
+inline FRHIDescriptorHandle GetBindlessResourceHandle(FUniformDataReader Reader, const FRHIUniformBufferResource& Resource)
+{
+	switch (Resource.MemberType)
+	{
+	case UBMT_TEXTURE:
+	{
+		FRHITexture* Texture = Reader.Read<FRHITexture*>(Resource);
+		return Texture ? Texture->GetDefaultBindlessHandle() : FRHIDescriptorHandle();
+	}
+	break;
+	case UBMT_SRV:
+	{
+		FRHIShaderResourceView* ShaderResourceView = Reader.Read<FRHIShaderResourceView*>(Resource);
+		return ShaderResourceView ? ShaderResourceView->GetBindlessHandle() : FRHIDescriptorHandle();
+	}
+	break;
+	case UBMT_SAMPLER:
+	{
+		FRHISamplerState* SamplerState = Reader.Read<FRHISamplerState*>(Resource);
+		return SamplerState ? SamplerState->GetBindlessHandle() : FRHIDescriptorHandle();
+	}
+	break;
+	case UBMT_RDG_TEXTURE:
+	{
+		FRDGTexture* RDGTexture = Reader.Read<FRDGTexture*>(Resource);
+		FRHITexture* Texture = RDGTexture ? RDGTexture->GetRHI() : nullptr;
+		return Texture ? Texture->GetDefaultBindlessHandle() : FRHIDescriptorHandle();
+	}
+	break;
+	case UBMT_RDG_TEXTURE_SRV:
+	case UBMT_RDG_BUFFER_SRV:
+	{
+		FRDGShaderResourceView* RDGShaderResourceView = Reader.Read<FRDGShaderResourceView*>(Resource);
+		if (RDGShaderResourceView)
+			RDGShaderResourceView->MarkResourceAsUsed();
+
+		FRHIShaderResourceView* ShaderResourceView = RDGShaderResourceView ? RDGShaderResourceView->GetRHI() : nullptr;
+		return ShaderResourceView ? ShaderResourceView->GetBindlessHandle() : FRHIDescriptorHandle();
+	}
+	break;
+	case UBMT_UAV:
+	{
+		FRHIUnorderedAccessView* UnorderedAccessView = Reader.Read<FRHIUnorderedAccessView*>(Resource);
+		return UnorderedAccessView ? UnorderedAccessView->GetBindlessHandle() : FRHIDescriptorHandle();
+	}
+	case UBMT_RDG_TEXTURE_UAV:
+	case UBMT_RDG_BUFFER_UAV:
+	{
+		FRDGUnorderedAccessView* RDGUnorderedAccessView = Reader.Read<FRDGUnorderedAccessView*>(Resource);
+		FRHIUnorderedAccessView* UnorderedAccessView = RDGUnorderedAccessView ? RDGUnorderedAccessView->GetRHI() : nullptr;
+		return UnorderedAccessView ? UnorderedAccessView->GetBindlessHandle() : FRHIDescriptorHandle();
+	}
+	break;
+	case UBMT_NESTED_STRUCT:
+	case UBMT_INCLUDED_STRUCT:
+	case UBMT_REFERENCED_STRUCT:
+	{
+		// Do nothing?
+	}
+	break;
+
+	default:
+		//checkf(false, TEXT("Unhandled resource type?"));
+		break;
+	}
+	return FRHIDescriptorHandle();
+}
+
+static void UpdateUniformBufferConstants(void* DestinationData, const void* SourceData, const FRHIUniformBufferLayout* Layout)
+{
+	check(DestinationData != nullptr);
+	check(SourceData != nullptr);
+
+	// First copy wholesale
+	FMemory::Memcpy(DestinationData, SourceData, Layout->ConstantBufferSize);
+
+	FUniformDataReader Reader(SourceData);
+
+	// Then copy indices over
+	for (const FRHIUniformBufferResource& Resource : Layout->Resources)
+	{
+		const FRHIDescriptorHandle Handle = GetBindlessResourceHandle(Reader, Resource);
+		if (Handle.IsValid())
+		{
+			const uint32 BindlessIndex = Handle.GetIndex();
+			FMemory::Memcpy(reinterpret_cast<uint8*>(DestinationData) + Resource.MemberOffset, &BindlessIndex, sizeof(BindlessIndex));
+		}
+	}
+}
+
 FUniformBufferRHIRef FD3D12DynamicRHI::RHICreateUniformBuffer(const void* Contents, const FRHIUniformBufferLayout* Layout, EUniformBufferUsage Usage, EUniformBufferValidation Validation)
 {
 	SCOPE_CYCLE_COUNTER(STAT_D3D12UpdateUniformBufferTime);
@@ -55,7 +160,8 @@ FUniformBufferRHIRef FD3D12DynamicRHI::RHICreateUniformBuffer(const void* Conten
 
 			// Copy the data to the upload heap
 			check(MappedData != nullptr);
-			FMemory::Memcpy(MappedData, Contents, NumBytesActualData);
+
+			UpdateUniformBufferConstants(MappedData, Contents, Layout);
 
 #if USE_STATIC_ROOT_SIGNATURE
 			NewUniformBuffer->View->Create(NewUniformBuffer->ResourceLocation.GetGPUVirtualAddress(), NumBytesActualData);
@@ -158,7 +264,7 @@ void FD3D12DynamicRHI::RHIUpdateUniformBuffer(FRHIUniformBuffer* UniformBufferRH
 
 		for (int32 Index = 0; Index < NumResources; ++Index)
 		{
-			const auto Parameter = Layout.Resources[Index];
+			const FRHIUniformBufferResource& Parameter = Layout.Resources[Index];
 			CmdListResources[Index] = GetShaderParameterResourceRHI(Contents, Parameter.MemberOffset, Parameter.MemberType);
 		}
 	}
@@ -189,7 +295,7 @@ void FD3D12DynamicRHI::RHIUpdateUniformBuffer(FRHIUniformBuffer* UniformBufferRH
 			}
 
 			check(MappedData != nullptr);
-			FMemory::Memcpy(MappedData, Contents, NumBytes);
+			UpdateUniformBufferConstants(MappedData, Contents, &Layout);
 		}
 
 		if (bBypass)
diff --git a/Engine/Source/Runtime/D3D12RHI/Private/D3D12Util.cpp b/Engine/Source/Runtime/D3D12RHI/Private/D3D12Util.cpp
--- a/Engine/Source/Runtime/D3D12RHI/Private/D3D12Util.cpp
+++ b/Engine/Source/Runtime/D3D12RHI/Private/D3D12Util.cpp
@@ -1072,6 +1072,7 @@ FD3D12QuantizedBoundShaderState GetRayTracingGlobalRootSignatureDesc()
 
 	OutQBSS.RootSignatureType = RS_RayTracingGlobal;
 	OutQBSS.bUseDiagnosticBuffer = true;
+	OutQBSS.bUseDirectlyIndexedSamplerHeap = true;
 
 	QBSSRegisterCounts.SamplerCount = MAX_SAMPLERS;
 	QBSSRegisterCounts.ShaderResourceCount = MAX_SRVS;
diff --git a/Engine/Source/Runtime/RenderCore/Private/RenderGraphUtils.cpp b/Engine/Source/Runtime/RenderCore/Private/RenderGraphUtils.cpp
--- a/Engine/Source/Runtime/RenderCore/Private/RenderGraphUtils.cpp
+++ b/Engine/Source/Runtime/RenderCore/Private/RenderGraphUtils.cpp
@@ -20,8 +20,9 @@ void ClearUnusedGraphResourcesImpl(
 	const auto& GraphResources = ParametersMetadata->GetLayout().GraphResources;
 
 	int32 ShaderResourceIndex = 0;
+	int32 BindlessResourceIndex = 0;
 	int32 GraphUniformBufferId = 0;
-	auto Base = reinterpret_cast<uint8*>(InoutParameters);
+	uint8* const Base = reinterpret_cast<uint8*>(InoutParameters);
 
 	for (int32 GraphResourceIndex = 0, GraphResourceCount = GraphResources.Num(); GraphResourceIndex < GraphResourceCount; GraphResourceIndex++)
 	{
@@ -34,7 +35,7 @@ void ClearUnusedGraphResourcesImpl(
 			Type == UBMT_RDG_BUFFER_SRV ||
 			Type == UBMT_RDG_BUFFER_UAV)
 		{
-			const auto& ResourceParameters = ShaderBindings.ResourceParameters;
+			const TMemoryImageArray<FShaderParameterBindings::FResourceParameter>& ResourceParameters = ShaderBindings.ResourceParameters;
 			const int32 ShaderResourceCount = ResourceParameters.Num();
 			for (; ShaderResourceIndex < ShaderResourceCount && ResourceParameters[ShaderResourceIndex].ByteOffset < ByteOffset; ++ShaderResourceIndex)
 			{
@@ -44,6 +45,17 @@ void ClearUnusedGraphResourcesImpl(
 			{
 				continue;
 			}
+
+			const TMemoryImageArray<FShaderParameterBindings::FBindlessResourceParameter>& BindlessResourceParameters = ShaderBindings.BindlessResourceParameters;
+			const int32 BindlessResourceCount = BindlessResourceParameters.Num();
+			for (; BindlessResourceIndex < BindlessResourceCount && BindlessResourceParameters[BindlessResourceIndex].ByteOffset < ByteOffset; BindlessResourceIndex++)
+			{
+			}
+
+			if (BindlessResourceIndex < BindlessResourceCount && BindlessResourceParameters[BindlessResourceIndex].ByteOffset == ByteOffset)
+			{
+				continue;
+			}
 		}
 		else if (Type == UBMT_RDG_UNIFORM_BUFFER)
 		{
diff --git a/Engine/Source/Runtime/RenderCore/Private/Shader.cpp b/Engine/Source/Runtime/RenderCore/Private/Shader.cpp
--- a/Engine/Source/Runtime/RenderCore/Private/Shader.cpp
+++ b/Engine/Source/Runtime/RenderCore/Private/Shader.cpp
@@ -611,9 +611,11 @@ void FShader::BuildParameterMapInfo(const TMap<FString, FParameterAllocation>& P
 		case EShaderParameterType::UniformBuffer:
 			UniformCount++;
 			break;
+		case EShaderParameterType::BindlessSamplerIndex:
 		case EShaderParameterType::Sampler:
 			SamplerCount++;
 			break;
+		case EShaderParameterType::BindlessResourceIndex:
 		case EShaderParameterType::SRV:
 			SRVCount++;
 			break;
@@ -634,6 +636,10 @@ void FShader::BuildParameterMapInfo(const TMap<FString, FParameterAllocation>& P
 			return &ParameterMapInfo.TextureSamplers;
 		case EShaderParameterType::SRV:
 			return &ParameterMapInfo.SRVs;
+		case EShaderParameterType::BindlessResourceIndex:
+			return &ParameterMapInfo.SRVs;
+		case EShaderParameterType::BindlessSamplerIndex:
+			return &ParameterMapInfo.TextureSamplers;
 		default:
 			return nullptr;
 		}
diff --git a/Engine/Source/Runtime/RenderCore/Private/ShaderCore.cpp b/Engine/Source/Runtime/RenderCore/Private/ShaderCore.cpp
--- a/Engine/Source/Runtime/RenderCore/Private/ShaderCore.cpp
+++ b/Engine/Source/Runtime/RenderCore/Private/ShaderCore.cpp
@@ -539,6 +539,19 @@ void FShaderParameterMap::RemoveParameterAllocation(const TCHAR* ParameterName)
 	ParameterMap.Remove(ParameterName);
 }
 
+TArray<FString> FShaderParameterMap::GetAllParameterNamesOfType(EShaderParameterType InType) const
+{
+	TArray<FString> Result;
+	for (const TMap<FString, FParameterAllocation>::ElementType& Parameter : ParameterMap)
+	{
+		if (Parameter.Value.Type == InType)
+		{
+			Result.Emplace(Parameter.Key);
+		}
+	}
+	return Result;
+}
+
 void FShaderCompilerDefinitions::SetFloatDefine(const TCHAR* Name, float Value)
 {
 	// Make sure the printed value perfectly matches the given number
diff --git a/Engine/Source/Runtime/RenderCore/Private/ShaderParameterStruct.cpp b/Engine/Source/Runtime/RenderCore/Private/ShaderParameterStruct.cpp
--- a/Engine/Source/Runtime/RenderCore/Private/ShaderParameterStruct.cpp
+++ b/Engine/Source/Runtime/RenderCore/Private/ShaderParameterStruct.cpp
@@ -199,24 +199,36 @@ struct FShaderParameterStructBindingContext
 				}
 				else if (bIsRHIResource || bIsRDGResource)
 				{
-					checkf(ParameterAllocation->BaseIndex < 256, TEXT("BaseIndex does not fit into uint8. Change FResourceParameter::BaseIndex type to uint16"));
-
-					FShaderParameterBindings::FResourceParameter Parameter;
-					Parameter.BaseIndex = (uint8)ParameterAllocation->BaseIndex;
-					Parameter.ByteOffset = ByteOffset + ArrayElementId * SHADER_PARAMETER_POINTER_ALIGNMENT;
-					Parameter.BaseType = BaseType;
+					if (ParameterAllocation->Type == EShaderParameterType::BindlessResourceIndex || ParameterAllocation->Type == EShaderParameterType::BindlessSamplerIndex)
+					{
+						FShaderParameterBindings::FBindlessResourceParameter Parameter;
+						Parameter.GlobalConstantOffset = ParameterAllocation->BaseIndex;
+						Parameter.ByteOffset = ByteOffset + ArrayElementId * SHADER_PARAMETER_POINTER_ALIGNMENT;
+						Parameter.BaseType = BaseType;
 
-					if (ParameterAllocation->Size != 1)
+						Bindings->BindlessResourceParameters.Add(Parameter);
+					}
+					else
 					{
-						UE_LOG(LogShaders, Fatal,
-							TEXT("Error with shader %s's (Permutation Id %d) parameter %s is %i bytes, cpp name = %s.")
-							TEXT("The shader compiler should give precisely which elements of an array did not get compiled out, ")
-							TEXT("for optimal automatic render graph pass dependency with ClearUnusedGraphResources()."),
-							Shader->GetTypeUnfrozen()->GetName(), PermutationId,
-							*ElementShaderBindingName, ParameterAllocation->Size, *CppName);
+						checkf(ParameterAllocation->BaseIndex < 256, TEXT("BaseIndex does not fit into uint8. Change FResourceParameter::BaseIndex type to uint16"));
+
+						FShaderParameterBindings::FResourceParameter Parameter;
+						Parameter.BaseIndex = (uint8)ParameterAllocation->BaseIndex;
+						Parameter.ByteOffset = ByteOffset + ArrayElementId * SHADER_PARAMETER_POINTER_ALIGNMENT;
+						Parameter.BaseType = BaseType;
+
+						if (ParameterAllocation->Size != 1)
+						{
+							UE_LOG(LogShaders, Fatal,
+								TEXT("Error with shader %s's (Permutation Id %d) parameter %s is %i bytes, cpp name = %s.")
+								TEXT("The shader compiler should give precisely which elements of an array did not get compiled out, ")
+								TEXT("for optimal automatic render graph pass dependency with ClearUnusedGraphResources()."),
+								Shader->GetTypeUnfrozen()->GetName(), PermutationId,
+								*ElementShaderBindingName, ParameterAllocation->Size, *CppName);
+						}
+
+						Bindings->ResourceParameters.Add(Parameter);
 					}
-
-					Bindings->ResourceParameters.Add(Parameter);
 				}
 				else
 				{
@@ -581,6 +593,74 @@ void ValidateShaderParameterResourcesRHI(const void* Contents, const FRHIUniform
 
 #endif // DO_CHECK
 
+inline FRHIDescriptorHandle GetBindlessParameterHandle(FShaderParameterReader Reader, const FShaderParameterBindings::FBindlessResourceParameter& Parameter)
+{
+	const EUniformBufferBaseType BaseType = static_cast<EUniformBufferBaseType>(Parameter.BaseType);
+	switch (BaseType)
+	{
+	case UBMT_TEXTURE:
+	{
+		FRHITexture* Texture = Reader.Read<FRHITexture*>(Parameter);
+		checkSlow(Texture);
+		return Texture ? Texture->GetDefaultBindlessHandle() : FRHIDescriptorHandle();
+	}
+	break;
+	case UBMT_SRV:
+	{
+		FRHIShaderResourceView* ShaderResourceView = Reader.Read<FRHIShaderResourceView*>(Parameter);
+		checkSlow(ShaderResourceView);
+		return ShaderResourceView->GetBindlessHandle();
+	}
+	break;
+	case UBMT_SAMPLER:
+	{
+		FRHISamplerState* SamplerState = Reader.Read<FRHISamplerState*>(Parameter);
+		checkSlow(SamplerState);
+		return SamplerState->GetBindlessHandle();
+	}
+	break;
+	case UBMT_RDG_TEXTURE:
+	{
+		FRDGTexture* RDGTexture = Reader.Read<FRDGTexture*>(Parameter);
+		checkSlow(RDGTexture);
+		FRHITexture* Texture = RDGTexture->GetRHI();
+		return Texture ? Texture->GetDefaultBindlessHandle() : FRHIDescriptorHandle();
+	}
+	break;
+	case UBMT_RDG_TEXTURE_SRV:
+	case UBMT_RDG_BUFFER_SRV:
+	{
+		FRDGShaderResourceView* RDGShaderResourceView = Reader.Read<FRDGShaderResourceView*>(Parameter);
+		checkSlow(RDGShaderResourceView);
+		RDGShaderResourceView->MarkResourceAsUsed();
+
+		FRHIShaderResourceView* ShaderResourceView = RDGShaderResourceView ? RDGShaderResourceView->GetRHI() : nullptr;
+		return ShaderResourceView ? ShaderResourceView->GetBindlessHandle() : FRHIDescriptorHandle();
+	}
+	break;
+	case UBMT_UAV:
+	{
+		FRHIUnorderedAccessView* UnorderedAccessView = Reader.Read<FRHIUnorderedAccessView*>(Parameter);
+		checkSlow(UnorderedAccessView);
+		return UnorderedAccessView->GetBindlessHandle();
+	}
+	case UBMT_RDG_TEXTURE_UAV:
+	case UBMT_RDG_BUFFER_UAV:
+	{
+		FRDGUnorderedAccessView* RDGUnorderedAccessView = Reader.Read<FRDGUnorderedAccessView*>(Parameter);
+		checkSlow(RDGUnorderedAccessView);
+
+		FRHIUnorderedAccessView* UnorderedAccessView = RDGUnorderedAccessView->GetRHI();
+		return UnorderedAccessView ? UnorderedAccessView->GetBindlessHandle() : FRHIDescriptorHandle();
+	}
+	break;
+	default:
+		checkf(false, TEXT("Unhandled resource type?"));
+		break;
+	}
+	return FRHIDescriptorHandle();
+}
+
 template<typename TRHICmdList>
 inline void SetShaderUAV(TRHICmdList& RHICmdList, FRHIGraphicsShader* ShaderRHI, FShaderParameterReader Reader, const FShaderParameterBindings::FResourceParameter& ParameterBinding)
 {
@@ -646,6 +726,16 @@ inline void SetShaderParametersInternal(
 		RHICmdList.SetShaderParameter(ShaderRHI, Parameter.BufferIndex, Parameter.BaseIndex, Parameter.ByteSize, DataPtr);
 	}
 
+	for (const FShaderParameterBindings::FBindlessResourceParameter& BindlessResource : Bindings.BindlessResourceParameters)
+	{
+		const FRHIDescriptorHandle Handle = GetBindlessParameterHandle(Reader, BindlessResource);
+		if (Handle.IsValid())
+		{
+			const uint32 BindlessIndex = Handle.GetIndex();
+			RHICmdList.SetShaderParameter(ShaderRHI, 0, BindlessResource.GlobalConstantOffset, 4, &BindlessIndex);
+		}
+	}
+
 	TArray<FShaderParameterBindings::FResourceParameter, TInlineAllocator<16>> GraphSRVs;
 
 	for (const FShaderParameterBindings::FResourceParameter& Parameter : Bindings.ResourceParameters)
diff --git a/Engine/Source/Runtime/RenderCore/Private/ShaderParameters.cpp b/Engine/Source/Runtime/RenderCore/Private/ShaderParameters.cpp
--- a/Engine/Source/Runtime/RenderCore/Private/ShaderParameters.cpp
+++ b/Engine/Source/Runtime/RenderCore/Private/ShaderParameters.cpp
@@ -162,9 +162,36 @@ static void CreateHLSLUniformBufferStructMembersDeclaration(
 		}
 		else if (IsShaderParameterTypeForUniformBufferLayout(Member.GetBaseType()))
 		{
-			// Skip resources, they will be replaced with padding by the next member in the constant buffer.
-			// This padding will cause gaps in the constant buffer.  Alternatively we could compact the constant buffer during RHICreateUniformBuffer.
-			continue;
+			// Generate the type dimensions for vectors and matrices.
+			uint32 HLSLMemberSize = 4;
+
+			const uint32 AbsoluteMemberOffset = StructOffset + Member.GetOffset();
+
+			// If the HLSL offset doesn't match the C++ offset, generate padding to fix it.
+			if (HLSLBaseOffset != AbsoluteMemberOffset)
+			{
+				check(HLSLBaseOffset < AbsoluteMemberOffset);
+				while (HLSLBaseOffset < AbsoluteMemberOffset)
+				{
+					Decl.ConstantBufferMembers += FString::Printf(TEXT("\t%s PrePadding_%s%u;\r\n"), *PreviousBaseTypeName, *NamePrefix, HLSLBaseOffset);
+					HLSLBaseOffset += 4;
+				};
+				check(HLSLBaseOffset == AbsoluteMemberOffset);
+			}
+			PreviousBaseTypeName = TEXT("uint");
+			HLSLBaseOffset = AbsoluteMemberOffset + HLSLMemberSize;
+
+			// Generate the member declaration.
+			const FString ParameterName = FString::Printf(TEXT("%s%s"), *NamePrefix, Member.GetName());
+
+			if (Member.GetBaseType() == UBMT_SAMPLER)
+			{
+				Decl.ConstantBufferMembers += FString::Printf(TEXT("\tDEFINE_SAMPLER_INDEX(%s);\r\n"), *ParameterName);
+			}
+			else
+			{
+				Decl.ConstantBufferMembers += FString::Printf(TEXT("\tDEFINE_RESOURCE_INDEX(%s);\r\n"), *ParameterName);
+			}
 		}
 		else 
 		{
@@ -243,22 +270,28 @@ static void CreateHLSLUniformBufferStructMembersDeclaration(
 
 		if (IsShaderParameterTypeForUniformBufferLayout(Member.GetBaseType()))
 		{
+			// TODO: handle arrays?
 			checkf(!IsRDGResourceAccessType(Member.GetBaseType()), TEXT("RDG access parameter types (e.g. RDG_TEXTURE_ACCESS) are not allowed in uniform buffers."));
-			if (Member.GetBaseType() == UBMT_SRV)
+
+			const FString ParameterName = FString::Printf(TEXT("%s%s"), *NamePrefix, Member.GetName());
+
+			if (Member.GetBaseType() == UBMT_SAMPLER)
+			{
+				Decl.ResourceMembers += FString::Printf(TEXT("DEFINE_SAMPLER_VARIABLE(%s, %s);\r\n"), Member.GetShaderType(), *ParameterName);
+				Decl.StructMembers += FString::Printf(TEXT("\t%s %s;\r\n"), Member.GetShaderType(), Member.GetName());
+				Decl.Initializer += FString::Printf(TEXT("GET_SAMPLER(%s),"), *ParameterName);
+			}
+			else if (Member.GetBaseType() == UBMT_SRV)
 			{
-				// TODO: handle arrays?
-				FString ParameterName = FString::Printf(TEXT("%s%s"),*NamePrefix,Member.GetName());
-				Decl.ResourceMembers += FString::Printf(TEXT("PLATFORM_SUPPORTS_SRV_UB_MACRO( %s %s; ) \r\n"), Member.GetShaderType(), *ParameterName);
+				Decl.ResourceMembers += FString::Printf(TEXT("PLATFORM_SUPPORTS_SRV_UB_MACRO( DEFINE_RESOURCE_VARIABLE(%s, %s); ) \r\n"), Member.GetShaderType(), *ParameterName);
 				Decl.StructMembers += FString::Printf(TEXT("\tPLATFORM_SUPPORTS_SRV_UB_MACRO( %s %s; ) \r\n"), Member.GetShaderType(), Member.GetName());
-				Decl.Initializer += FString::Printf(TEXT(" PLATFORM_SUPPORTS_SRV_UB_MACRO( %s, ) "), *ParameterName);
+				Decl.Initializer += FString::Printf(TEXT(" PLATFORM_SUPPORTS_SRV_UB_MACRO( GET_RESOURCE(%s), ) "), *ParameterName);
 			}
 			else
 			{
-				// TODO: handle arrays?
-				FString ParameterName = FString::Printf(TEXT("%s%s"),*NamePrefix,Member.GetName());
-				Decl.ResourceMembers += FString::Printf(TEXT("%s %s;\r\n"), Member.GetShaderType(), *ParameterName);
+				Decl.ResourceMembers += FString::Printf(TEXT("DEFINE_RESOURCE_VARIABLE(%s, %s);\r\n"), Member.GetShaderType(), *ParameterName);
 				Decl.StructMembers += FString::Printf(TEXT("\t%s %s;\r\n"), Member.GetShaderType(), Member.GetName());
-				Decl.Initializer += FString::Printf(TEXT("%s,"), *ParameterName);
+				Decl.Initializer += FString::Printf(TEXT("GET_RESOURCE(%s),"), *ParameterName);
 			}
 		}
 	}
diff --git a/Engine/Source/Runtime/RenderCore/Public/Shader.h b/Engine/Source/Runtime/RenderCore/Public/Shader.h
--- a/Engine/Source/Runtime/RenderCore/Public/Shader.h
+++ b/Engine/Source/Runtime/RenderCore/Public/Shader.h
@@ -597,16 +597,21 @@ public:
 		LAYOUT_FIELD(uint16, ByteOffset);
 		LAYOUT_FIELD(uint8, BaseIndex);
 		LAYOUT_FIELD(EUniformBufferBaseType, BaseType);
+	};
 
+	struct FBindlessResourceParameter
+	{
+		DECLARE_INLINE_TYPE_LAYOUT(FBindlessResourceParameter, NonVirtual);
+		LAYOUT_FIELD(uint16, ByteOffset);
+		LAYOUT_FIELD(uint16, GlobalConstantOffset);
+		LAYOUT_FIELD(EUniformBufferBaseType, BaseType);
 	};
 
 	struct FParameterStructReference
 	{
 		DECLARE_INLINE_TYPE_LAYOUT(FParameterStructReference, NonVirtual);
-
 		LAYOUT_FIELD(uint16, BufferIndex);
 		LAYOUT_FIELD(uint16, ByteOffset);
-
 	};
 
 	DECLARE_TYPE_LAYOUT(FShaderParameterBindings, NonVirtual);
@@ -618,6 +623,7 @@ public:
 
 	LAYOUT_FIELD(TMemoryImageArray<FParameter>, Parameters);
 	LAYOUT_FIELD(TMemoryImageArray<FResourceParameter>, ResourceParameters);
+	LAYOUT_FIELD(TMemoryImageArray<FBindlessResourceParameter>, BindlessResourceParameters);
 	LAYOUT_FIELD(TMemoryImageArray<FParameterStructReference>, GraphUniformBuffers);
 	LAYOUT_FIELD(TMemoryImageArray<FParameterStructReference>, ParameterReferences);
 
diff --git a/Engine/Source/Runtime/RenderCore/Public/ShaderCore.h b/Engine/Source/Runtime/RenderCore/Public/ShaderCore.h
--- a/Engine/Source/Runtime/RenderCore/Public/ShaderCore.h
+++ b/Engine/Source/Runtime/RenderCore/Public/ShaderCore.h
@@ -183,6 +183,9 @@ enum class EShaderParameterType : uint8
 	SRV,
 	UAV,
 
+	BindlessResourceIndex,
+	BindlessSamplerIndex,
+
 	Num
 };
 
@@ -238,6 +241,9 @@ public:
 	RENDERCORE_API void AddParameterAllocation(const TCHAR* ParameterName,uint16 BufferIndex,uint16 BaseIndex,uint16 Size,EShaderParameterType ParameterType);
 	RENDERCORE_API void RemoveParameterAllocation(const TCHAR* ParameterName);
 
+	/** Returns an array of all parameters with the given type. */
+	RENDERCORE_API TArray<FString> GetAllParameterNamesOfType(EShaderParameterType InType) const;
+
 	/** Checks that all parameters are bound and asserts if any aren't in a debug build
 	* @param InVertexFactoryType can be 0
 	*/
@@ -595,13 +601,21 @@ struct FSharedShaderCompilerEnvironment final : public FShaderCompilerEnvironmen
 	virtual ~FSharedShaderCompilerEnvironment() = default;
 };
 
+enum class EShaderResourceUsageFlags : uint8
+{
+	GlobalUniformBuffer = 1 << 0,
+	BindlessResources   = 1 << 1,
+	BindlessSamplers    = 1 << 2,
+};
+ENUM_CLASS_FLAGS(EShaderResourceUsageFlags)
+
 // if this changes you need to make sure all shaders get invalidated
 struct FShaderCodePackedResourceCounts
 {
 	// for FindOptionalData() and AddOptionalData()
 	static const uint8 Key = 'p';
 
-	bool bGlobalUniformBufferUsed;
+	EShaderResourceUsageFlags UsageFlags;
 	uint8 NumSamplers;
 	uint8 NumSRVs;
 	uint8 NumCBs;
diff --git a/Engine/Source/Runtime/Windows/D3D11RHI/Private/D3D11Shaders.cpp b/Engine/Source/Runtime/Windows/D3D11RHI/Private/D3D11Shaders.cpp
--- a/Engine/Source/Runtime/Windows/D3D11RHI/Private/D3D11Shaders.cpp
+++ b/Engine/Source/Runtime/Windows/D3D11RHI/Private/D3D11Shaders.cpp
@@ -23,7 +23,7 @@ static inline void ReadShaderOptionalData(FShaderCodeReader& InShaderCode, TShad
 		OutShader.UAVMask = ResourceMasks->UAVMask;
 	}
 	
-	OutShader.bShaderNeedsGlobalConstantBuffer = PackedResourceCounts->bGlobalUniformBufferUsed;
+	OutShader.bShaderNeedsGlobalConstantBuffer = EnumHasAnyFlags(PackedResourceCounts->UsageFlags, EShaderResourceUsageFlags::GlobalUniformBuffer);
 #if RHI_INCLUDE_SHADER_DEBUG_DATA
 	OutShader.ShaderName = InShaderCode.FindOptionalData(FShaderCodeName::Key);
 
@@ -48,7 +48,7 @@ static inline void ReadShaderOptionalData(FShaderCodeReader& InShaderCode, TShad
 		FBufferReader Ar((void*)VendorExtensionData, VendorExtensionTableSize, false);
 		Ar << OutShader.VendorExtensions;
 	}
-	OutShader.bShaderNeedsGlobalConstantBuffer = PackedResourceCounts->bGlobalUniformBufferUsed;
+	OutShader.bShaderNeedsGlobalConstantBuffer = EnumHasAnyFlags(PackedResourceCounts->UsageFlags, EShaderResourceUsageFlags::GlobalUniformBuffer);
 
 	int32 IsSm6ShaderSize = 1;
 	const uint8* IsSm6Shader = InShaderCode.FindOptionalData('6', IsSm6ShaderSize);
