Moving ShaderParameterParser into its own source files.

diff --git a/Engine/Source/Developer/Apple/MetalShaderFormat/Private/MetalShaderCompiler.cpp b/Engine/Source/Developer/Apple/MetalShaderFormat/Private/MetalShaderCompiler.cpp
--- a/Engine/Source/Developer/Apple/MetalShaderFormat/Private/MetalShaderCompiler.cpp
+++ b/Engine/Source/Developer/Apple/MetalShaderFormat/Private/MetalShaderCompiler.cpp
@@ -5,6 +5,7 @@
 #include "ShaderCore.h"
 #include "MetalShaderResources.h"
 #include "ShaderCompilerCommon.h"
+#include "ShaderParameterParser.h"
 #include "Misc/ConfigCacheIni.h"
 #include "Serialization/MemoryWriter.h"
 #include "Serialization/MemoryReader.h"
diff --git a/Engine/Source/Developer/ShaderCompilerCommon/Private/ShaderCompilerCommon.cpp b/Engine/Source/Developer/ShaderCompilerCommon/Private/ShaderCompilerCommon.cpp
--- a/Engine/Source/Developer/ShaderCompilerCommon/Private/ShaderCompilerCommon.cpp
+++ b/Engine/Source/Developer/ShaderCompilerCommon/Private/ShaderCompilerCommon.cpp
@@ -1,6 +1,7 @@
 // Copyright Epic Games, Inc. All Rights Reserved.
 
 #include "ShaderCompilerCommon.h"
+#include "ShaderParameterParser.h"
 #include "Misc/Paths.h"
 #include "Modules/ModuleManager.h"
 #include "HlslccDefinitions.h"
@@ -637,31 +638,6 @@ void AddNoteToDisplayShaderParameterStructureOnCppSide(
 	CompilerOutput.Errors.Add(Error);
 }
 
-void AddNoteToDisplayShaderParameterMemberOnCppSide(
-	const FShaderCompilerInput& CompilerInput,
-	const FShaderParameterParser::FParsedShaderParameter& ParsedParameter,
-	FShaderCompilerOutput& CompilerOutput)
-{
-	const FShaderParametersMetadata* MemberContainingStruct = nullptr;
-	const FShaderParametersMetadata::FMember* Member = nullptr;
-	{
-		int32 ArrayElementId = 0;
-		FString NamePrefix;
-		CompilerInput.RootParametersStructure->FindMemberFromOffset(ParsedParameter.ConstantBufferOffset, &MemberContainingStruct, &Member, &ArrayElementId, &NamePrefix);
-	}
-
-	FString CppCodeName = CompilerInput.RootParametersStructure->GetFullMemberCodeName(ParsedParameter.ConstantBufferOffset);
-
-	FShaderCompilerError Error;
-	Error.StrippedErrorMessage = FString::Printf(
-		TEXT("Note: Definition of %s"),
-		*CppCodeName);
-	Error.ErrorVirtualFilePath = ANSI_TO_TCHAR(MemberContainingStruct->GetFileName());
-	Error.ErrorLineString = FString::FromInt(Member->GetFileLine());
-
-	CompilerOutput.Errors.Add(Error);
-}
-
 void AddUnboundShaderParameterError(
 	const FShaderCompilerInput& CompilerInput,
 	const FShaderParameterParser& ShaderParameterParser,
@@ -686,798 +662,6 @@ void AddUnboundShaderParameterError(
 	AddNoteToDisplayShaderParameterStructureOnCppSide(CompilerInput.RootParametersStructure, CompilerOutput);
 }
 
-inline bool MemberWasPotentiallyMoved(const FShaderParametersMetadata::FMember& InMember)
-{
-	const EUniformBufferBaseType BaseType = InMember.GetBaseType();
-	return BaseType == UBMT_INT32 || BaseType == UBMT_UINT32 || BaseType == UBMT_FLOAT32
-		|| BaseType == UBMT_TEXTURE || BaseType == UBMT_SRV || BaseType == UBMT_UAV
-		|| BaseType == UBMT_RDG_TEXTURE || BaseType == UBMT_RDG_TEXTURE_SRV || BaseType == UBMT_RDG_TEXTURE_UAV
-		|| BaseType == UBMT_RDG_BUFFER_SRV || BaseType == UBMT_RDG_BUFFER_UAV
-		|| BaseType == UBMT_SAMPLER
-		;
-}
-
-bool FShaderParameterParser::ParseAndMoveShaderParametersToRootConstantBuffer(
-	const FShaderCompilerInput& CompilerInput,
-	FShaderCompilerOutput& CompilerOutput,
-	FString& PreprocessedShaderSource,
-	const TCHAR* ConstantBufferType)
-{
-	// The shader doesn't have any parameter binding through shader structure, therefore don't do anything.
-	if (!CompilerInput.RootParametersStructure)
-	{
-		return true;
-	}
-
-	const bool bMoveToRootConstantBuffer = ConstantBufferType != nullptr;
-	OriginalParsedShader = PreprocessedShaderSource;
-
-	// Reserves the number of parameters up front.
-	ParsedParameters.Reserve(CompilerInput.RootParametersStructure->GetSize() / sizeof(int32));
-
-	CompilerInput.RootParametersStructure->IterateShaderParameterMembers(
-		[&](const FShaderParametersMetadata& ParametersMetadata,
-			const FShaderParametersMetadata::FMember& Member,
-			const TCHAR* ShaderBindingName,
-			uint16 ByteOffset)
-	{
-		FParsedShaderParameter ParsedParameter;
-		ParsedParameter.Member = &Member;
-		ParsedParameter.ConstantBufferOffset = ByteOffset;
-		check(ParsedParameter.IsBindable());
-
-		ParsedParameters.Add(ShaderBindingName, ParsedParameter);
-	});
-
-	bool bSuccess = true;
-
-	// Browse the code for global shader parameter, Save their type and erase them white spaces.
-	{
-		enum class EState
-		{
-			// When to look for something to scan.
-			Scanning,
-
-			// When going to next ; in the global scope and reset.
-			GoToNextSemicolonAndReset,
-
-			// Parsing what might be a type of the parameter.
-			ParsingPotentialType,
-			FinishedPotentialType,
-
-			// Parsing what might be a name of the parameter.
-			ParsingPotentialName,
-			FinishedPotentialName,
-
-			// Parsing what looks like array of the parameter.
-			ParsingPotentialArraySize,
-			FinishedArraySize,
-
-			// Found a parameter, just finish to it's semi colon.
-			FoundParameter,
-		};
-
-		const int32 ShaderSourceLen = PreprocessedShaderSource.Len();
-
-		int32 CurrentPragamLineoffset = -1;
-		int32 CurrentLineoffset = 0;
-
-		int32 TypeQualifierStartPos = -1;
-		int32 TypeStartPos = -1;
-		int32 TypeEndPos = -1;
-		int32 NameStartPos = -1;
-		int32 NameEndPos = -1;
-		int32 ArrayStartPos = -1;
-		int32 ArrayEndPos = -1;
-		int32 ScopeIndent = 0;
-
-		EState State = EState::Scanning;
-		bool bGoToNextLine = false;
-
-		auto ResetState = [&]()
-		{
-			TypeQualifierStartPos = -1;
-			TypeStartPos = -1;
-			TypeEndPos = -1;
-			NameStartPos = -1;
-			NameEndPos = -1;
-			ArrayStartPos = -1;
-			ArrayEndPos = -1;
-			State = EState::Scanning;
-		};
-
-		auto EmitError = [&](const FString& ErrorMessage)
-		{
-			FShaderCompilerError Error;
-			Error.StrippedErrorMessage = ErrorMessage;
-			ExtractFileAndLine(CurrentPragamLineoffset, CurrentLineoffset, Error.ErrorVirtualFilePath, Error.ErrorLineString);
-			CompilerOutput.Errors.Add(Error);
-			bSuccess = false;
-		};
-
-		auto EmitUnpextectedHLSLSyntaxError = [&]()
-		{
-			EmitError(TEXT("Unexpected syntax when parsing shader parameters from shader code."));
-			State = EState::GoToNextSemicolonAndReset;
-		};
-
-		for (int32 Cursor = 0; Cursor < ShaderSourceLen; Cursor++)
-		{
-			const TCHAR Char = PreprocessedShaderSource[Cursor];
-
-			auto FoundShaderParameter = [&]()
-			{
-				check(Char == ';');
-				check(TypeStartPos != -1);
-				check(TypeEndPos != -1);
-				check(NameStartPos != -1);
-				check(NameEndPos != -1);
-
-				FString Type = PreprocessedShaderSource.Mid(TypeStartPos, TypeEndPos - TypeStartPos + 1);
-				FString Name = PreprocessedShaderSource.Mid(NameStartPos, NameEndPos - NameStartPos + 1);
-				const FString ParsedName = Name;
-
-				const bool bBindlessParameter = UE::ShaderCompilerCommon::RemoveBindlessParameterPrefix(Name);
-
-				FParsedShaderParameter ParsedParameter;
-				bool bUpdateParsedParameters = false;
-				bool bEraseOriginalParameter = false;
-				if (ParsedParameters.Contains(Name))
-				{
-					if (ParsedParameters.FindChecked(Name).IsFound())
-					{
-						// If it has already been found, it means it is duplicated. Do nothing and let the shader compiler throw the error.
-					}
-					else
-					{
-						// Update the parsed parameters
-						bUpdateParsedParameters = true;
-						ParsedParameter = ParsedParameters.FindChecked(Name);
-
-						// Erase the parameter to move it into the root constant buffer.
-						if (bMoveToRootConstantBuffer && ParsedParameter.IsBindable())
-						{
-							EUniformBufferBaseType BaseType = ParsedParameter.Member->GetBaseType();
-							bEraseOriginalParameter = BaseType == UBMT_INT32 || BaseType == UBMT_UINT32 || BaseType == UBMT_FLOAT32 || bBindlessParameter;
-						}
-					}
-				}
-				else
-				{
-					// Update the parsed parameters to still have file and line number.
-					bUpdateParsedParameters = true;
-				}
-
-				// Update 
-				if (bUpdateParsedParameters)
-				{
-					ParsedParameter.ParsedName = ParsedName;
-					ParsedParameter.ParsedType = Type;
-					ParsedParameter.ParsedPragmaLineoffset = CurrentPragamLineoffset;
-					ParsedParameter.ParsedLineOffset = CurrentLineoffset;
-					ParsedParameter.bOriginalParameterErased = bEraseOriginalParameter;
-
-					if (ArrayStartPos != -1 && ArrayEndPos != -1)
-					{
-						ParsedParameter.ParsedArraySize = PreprocessedShaderSource.Mid(ArrayStartPos + 1, ArrayEndPos - ArrayStartPos - 1);
-					}
-
-					ParsedParameters.Add(Name, ParsedParameter);
-				}
-
-				// Erases this shader parameter conserving the same line numbers.
-				if (bEraseOriginalParameter)
-				{
-					for (int32 j = (TypeQualifierStartPos != -1 ? TypeQualifierStartPos : TypeStartPos); j <= Cursor; j++)
-					{
-						if (PreprocessedShaderSource[j] != '\r' && PreprocessedShaderSource[j] != '\n')
-							PreprocessedShaderSource[j] = ' ';
-					}
-				}
-
-				ResetState();
-			};
-
-			const bool bIsWhiteSpace = Char == ' ' || Char == '\t' || Char == '\r' || Char == '\n';
-			const bool bIsLetter = (Char >= 'a' && Char <= 'z') || (Char >= 'A' && Char <= 'Z');
-			const bool bIsNumber = Char >= '0' && Char <= '9';
-
-			const TCHAR* UpComing = (*PreprocessedShaderSource) + Cursor;
-			const int32 RemainingSize = ShaderSourceLen - Cursor;
-
-			CurrentLineoffset += Char == '\n';
-
-			// Go to the next line if this is a preprocessor macro.
-			if (bGoToNextLine)
-			{
-				if (Char == '\n')
-				{
-					bGoToNextLine = false;
-				}
-				continue;
-			}
-			else if (Char == '#')
-			{
-				if (RemainingSize > 6 && FCString::Strncmp(UpComing, TEXT("#line "), 6) == 0)
-				{
-					CurrentPragamLineoffset = Cursor;
-					CurrentLineoffset = -1; // that will be incremented to 0 when reaching the \n at the end of the #line
-				}
-
-				bGoToNextLine = true;
-				continue;
-			}
-
-			// If within a scope, just carry on until outside the scope.
-			if (ScopeIndent > 0 || Char == '{')
-			{
-				if (Char == '{')
-				{
-					ScopeIndent++;
-				}
-				else if (Char == '}')
-				{
-					ScopeIndent--;
-					if (ScopeIndent == 0)
-					{
-						ResetState();
-					}
-				}
-				continue;
-			}
-
-			if (State == EState::Scanning)
-			{
-				if (bIsLetter)
-				{
-					static const TCHAR* KeywordTable[] = {
-						TEXT("enum"),
-						TEXT("class"),
-						TEXT("const"),
-						TEXT("struct"),
-						TEXT("static"),
-					};
-					static int32 KeywordTableSize[] = {4, 5, 5, 6, 6};
-
-					int32 RecognisedKeywordId = -1;
-					for (int32 KeywordId = 0; KeywordId < UE_ARRAY_COUNT(KeywordTable); KeywordId++)
-					{
-						const TCHAR* Keyword = KeywordTable[KeywordId];
-						const int32 KeywordSize = KeywordTableSize[KeywordId];
-
-						if (RemainingSize > KeywordSize)
-						{
-							TCHAR KeywordEndTestChar = UpComing[KeywordSize];
-
-							if ((KeywordEndTestChar == ' ' || KeywordEndTestChar == '\r' || KeywordEndTestChar == '\n' || KeywordEndTestChar == '\t') &&
-								FCString::Strncmp(UpComing, Keyword, KeywordSize) == 0)
-							{
-								RecognisedKeywordId = KeywordId;
-								break;
-							}
-						}
-					}
-
-					if (RecognisedKeywordId == -1)
-					{
-						// Might have found beginning of the type of a parameter.
-						State = EState::ParsingPotentialType;
-						TypeStartPos = Cursor;
-					}
-					else if (RecognisedKeywordId == 2)
-					{
-						// Ignore the const keywords, but still parse given it might still be a shader parameter.
-						if (TypeQualifierStartPos == -1)
-						{
-							// If the parameter is erased, we also have to erase *all* 'const'-qualifiers, e.g. "const int Foo" or "const const int Foo".
-							TypeQualifierStartPos = Cursor;
-						}
-						Cursor += KeywordTableSize[RecognisedKeywordId];
-					}
-					else
-					{
-						// Purposefully ignore enum, class, struct, static
-						State = EState::GoToNextSemicolonAndReset;
-					}
-				}
-				else if (bIsWhiteSpace)
-				{
-					// Keep parsing void.
-				}
-				else if (Char == ';')
-				{
-					// Looks like redundant semicolon, just ignore and keep scanning.
-				}
-				else
-				{
-					// No idea what this is, just go to next semi colon.
-					State = EState::GoToNextSemicolonAndReset;
-				}
-			}
-			else if (State == EState::GoToNextSemicolonAndReset)
-			{
-				// If need to go to next global semicolon and reach it. Resume browsing.
-				if (Char == ';')
-				{
-					ResetState();
-				}
-			}
-			else if (State == EState::ParsingPotentialType)
-			{
-				// Found character legal for a type...
-				if (bIsLetter ||
-					bIsNumber ||
-					Char == '<' || Char == '>' || Char == '_')
-				{
-					// Keep browsing what might be type of the parameter.
-				}
-				else if (bIsWhiteSpace)
-				{
-					// Might have found a type.
-					State = EState::FinishedPotentialType;
-					TypeEndPos = Cursor - 1;
-				}
-				else
-				{
-					// Found unexpected character in the type.
-					State = EState::GoToNextSemicolonAndReset;
-				}
-			}
-			else if (State == EState::FinishedPotentialType)
-			{
-				if (bIsLetter)
-				{
-					// Might have found beginning of the name of a parameter.
-					State = EState::ParsingPotentialName;
-					NameStartPos = Cursor;
-				}
-				else if (bIsWhiteSpace)
-				{
-					// Keep parsing void.
-				}
-				else
-				{
-					// No idea what this is, just go to next semi colon.
-					State = EState::GoToNextSemicolonAndReset;
-				}
-			}
-			else if (State == EState::ParsingPotentialName)
-			{
-				// Found character legal for a name...
-				if (bIsLetter ||
-					bIsNumber ||
-					Char == '_')
-				{
-					// Keep browsing what might be name of the parameter.
-				}
-				else if (Char == ':' || Char == '=')
-				{
-					// Found a parameter with syntax:
-					// uint MyParameter : <whatever>;
-					// uint MyParameter = <DefaultValue>;
-					NameEndPos = Cursor - 1;
-					State = EState::FoundParameter;
-				}
-				else if (Char == ';')
-				{
-					// Found a parameter with syntax:
-					// uint MyParameter;
-					NameEndPos = Cursor - 1;
-					FoundShaderParameter();
-				}
-				else if (Char == '[')
-				{
-					// Syntax:
-					//  uint MyArray[
-					NameEndPos = Cursor - 1;
-					ArrayStartPos = Cursor;
-					State = EState::ParsingPotentialArraySize;
-				}
-				else if (bIsWhiteSpace)
-				{
-					// Might have found a name.
-					// uint MyParameter <Still need to know what is after>;
-					NameEndPos = Cursor - 1;
-					State = EState::FinishedPotentialName;
-				}
-				else
-				{
-					// Found unexpected character in the name.
-					// syntax:
-					// uint MyFunction(<Don't care what is after>
-					State = EState::GoToNextSemicolonAndReset;
-				}
-			}
-			else if (
-				State == EState::FinishedPotentialName ||
-				State == EState::FinishedArraySize)
-			{
-				if (Char == ';')
-				{
-					// Found a parameter with syntax:
-					// uint MyParameter <a bit of OK stuf>;
-					FoundShaderParameter();
-				}
-				else if (Char == ':')
-				{
-					// Found a parameter with syntax:
-					// uint MyParameter <a bit of OK stuf> : <Ignore all this crap>;
-					State = EState::FoundParameter;
-				}
-				else if (Char == '=')
-				{
-					// Found syntax that doesn't make any sens:
-					// uint MyParameter <a bit of OK stuf> = <Ignore all this crap>;
-					State = EState::FoundParameter;
-					// TDOO: should error out that this is useless.
-				}
-				else if (Char == '[')
-				{
-					if (State == EState::FinishedPotentialName)
-					{
-						// Syntax:
-						//  uint MyArray [
-						ArrayStartPos = Cursor;
-						State = EState::ParsingPotentialArraySize;
-					}
-					else
-					{
-						EmitError(TEXT("Shader parameters can only support one dimensional array"));
-					}
-				}
-				else if (bIsWhiteSpace)
-				{
-					// Keep parsing void.
-				}
-				else
-				{
-					// Found unexpected stuff.
-					State = EState::GoToNextSemicolonAndReset;
-				}
-			}
-			else if (State == EState::ParsingPotentialArraySize)
-			{
-				if (Char == ']')
-				{
-					ArrayEndPos = Cursor;
-					State = EState::FinishedArraySize;
-				}
-				else if (Char == ';')
-				{
-					EmitUnpextectedHLSLSyntaxError();
-				}
-				else
-				{
-					// Keep going through the array size that might be a complex expression.
-				}
-			}
-			else if (State == EState::FoundParameter)
-			{
-				if (Char == ';')
-				{
-					FoundShaderParameter();
-				}
-				else
-				{
-					// Cary on skipping all crap we don't care about shader parameter until we find it's semi colon.
-				}
-			}
-			else
-			{
-				unimplemented();
-			}
-		} // for (int32 Cursor = 0; Cursor < PreprocessedShaderSource.Len(); Cursor++)
-	}
-
-	// Generate the root cbuffer content.
-	if (bMoveToRootConstantBuffer)
-	{
-		FString RootCBufferContent;
-
-		CompilerInput.RootParametersStructure->IterateShaderParameterMembers(
-			[&](const FShaderParametersMetadata& ParametersMetadata,
-				const FShaderParametersMetadata::FMember& Member,
-				const TCHAR* ShaderBindingName,
-				uint16 ByteOffset)
-		{
-			if (MemberWasPotentiallyMoved(Member))
-			{
-				FParsedShaderParameter* ParsedParameter = ParsedParameters.Find(ShaderBindingName);
-				if (ParsedParameter && ParsedParameter->IsFound() && ParsedParameter->bOriginalParameterErased)
-				{
-					uint32 ConstantRegister = ByteOffset / 16;
-					const TCHAR* ConstantSwizzle = [ByteOffset]()
-						{
-							switch (ByteOffset % 16)
-							{
-							default: unimplemented();
-							case 0:  return TEXT("");
-							case 4:  return TEXT(".y");
-							case 8:  return TEXT(".z");
-							case 12: return TEXT(".w");
-							}
-						}();
-
-					if (!ParsedParameter->ParsedArraySize.IsEmpty())
-					{
-						RootCBufferContent.Append(FString::Printf(
-							TEXT("%s %s[%s] : packoffset(c%d%s);\r\n"),
-							*ParsedParameter->ParsedType,
-							*ParsedParameter->ParsedName,
-							*ParsedParameter->ParsedArraySize,
-							ConstantRegister,
-							ConstantSwizzle));
-					}
-					else
-					{
-						RootCBufferContent.Append(FString::Printf(
-							TEXT("%s %s : packoffset(c%d%s);\r\n"),
-							*ParsedParameter->ParsedType,
-							*ParsedParameter->ParsedName,
-							ConstantRegister,
-							ConstantSwizzle));
-					}
-				}
-			}
-		});
-
-		FString CBufferCodeBlock = FString::Printf(
-			TEXT("%s %s\r\n")
-			TEXT("{\r\n")
-			TEXT("%s")
-			TEXT("}\r\n\r\n"),
-			ConstantBufferType,
-			FShaderParametersMetadata::kRootUniformBufferBindingName,
-			*RootCBufferContent);
-
-		FString NewShaderCode = (
-			MakeInjectedShaderCodeBlock(TEXT("ParseAndMoveShaderParametersToRootConstantBuffer"), CBufferCodeBlock) +
-			PreprocessedShaderSource);
-
-		PreprocessedShaderSource = MoveTemp(NewShaderCode);
-
-		bMovedLoosedParametersToRootConstantBuffer = true;
-	} // if (bMoveToRootConstantBuffer)
-
-	return bSuccess;
-}
-
-void FShaderParameterParser::ValidateShaderParameterType(
-	const FShaderCompilerInput& CompilerInput,
-	const FString& ShaderBindingName,
-	int32 ReflectionOffset,
-	int32 ReflectionSize,
-	bool bPlatformSupportsPrecisionModifier,
-	FShaderCompilerOutput& CompilerOutput) const
-{
-	FString BindingName(ShaderBindingName);
-
-	const bool bBindlessHack = UE::ShaderCompilerCommon::RemoveBindlessParameterPrefix(BindingName);
-
-	const FShaderParameterParser::FParsedShaderParameter& ParsedParameter = FindParameterInfos(BindingName);
-
-	check(ParsedParameter.IsFound());
-	check(CompilerInput.RootParametersStructure);
-
-	if (ReflectionSize > 0 && bMovedLoosedParametersToRootConstantBuffer)
-	{
-		// Verify the offset of the parameter coming from shader reflections honor the packoffset()
-		check(ReflectionOffset == ParsedParameter.ConstantBufferOffset);
-	}
-
-	// Validate the shader type.
-	if (!bBindlessHack)
-	{
-		FString ExpectedShaderType;
-		ParsedParameter.Member->GenerateShaderParameterType(ExpectedShaderType, bPlatformSupportsPrecisionModifier);
-
-		const bool bShouldBeInt = ParsedParameter.Member->GetBaseType() == UBMT_INT32;
-		const bool bShouldBeUint = ParsedParameter.Member->GetBaseType() == UBMT_UINT32;
-
-		// Match parsed type with expected shader type
-		bool bIsTypeCorrect = ParsedParameter.ParsedType == ExpectedShaderType;
-
-		if (!bIsTypeCorrect)
-		{
-			// Accept half-precision floats when single-precision was requested
-			if (ParsedParameter.ParsedType.StartsWith(TEXT("half")) && ParsedParameter.Member->GetBaseType() == UBMT_FLOAT32)
-			{
-				bIsTypeCorrect = (FCString::Strcmp(*ParsedParameter.ParsedType + 4, *ExpectedShaderType + 5) == 0);
-			}
-			// Accept single-precision floats when half-precision was expected
-			else if (ParsedParameter.ParsedType.StartsWith(TEXT("float")) && ExpectedShaderType.StartsWith(TEXT("half")))
-			{
-				bIsTypeCorrect = (FCString::Strcmp(*ParsedParameter.ParsedType + 5, *ExpectedShaderType + 4) == 0);
-			}
-			// support for min16float
-			else if (ParsedParameter.ParsedType.StartsWith(TEXT("min16float")) && ExpectedShaderType.StartsWith(TEXT("float")))
-			{
-				bIsTypeCorrect = (FCString::Strcmp(*ParsedParameter.ParsedType + 10, *ExpectedShaderType + 5) == 0);
-			}
-			else if (ParsedParameter.ParsedType.StartsWith(TEXT("min16float")) && ExpectedShaderType.StartsWith(TEXT("half")))
-			{
-				bIsTypeCorrect = (FCString::Strcmp(*ParsedParameter.ParsedType + 10, *ExpectedShaderType + 4) == 0);
-			}
-		}
-
-		// Allow silent casting between signed and unsigned on shader bindings.
-		if (!bIsTypeCorrect && (bShouldBeInt || bShouldBeUint))
-		{
-			FString NewExpectedShaderType;
-			if (bShouldBeInt)
-			{
-				// tries up with an uint.
-				NewExpectedShaderType = TEXT("u") + ExpectedShaderType;
-			}
-			else
-			{
-				// tries up with an int.
-				NewExpectedShaderType = ExpectedShaderType;
-				NewExpectedShaderType.RemoveAt(0);
-			}
-
-			bIsTypeCorrect = ParsedParameter.ParsedType == NewExpectedShaderType;
-		}
-
-		if (!bIsTypeCorrect)
-		{
-			FString CppCodeName = CompilerInput.RootParametersStructure->GetFullMemberCodeName(ParsedParameter.ConstantBufferOffset);
-
-			FShaderCompilerError Error;
-			Error.StrippedErrorMessage = FString::Printf(
-				TEXT("Error: Type %s of shader parameter %s in shader mismatch the shader parameter structure: %s expects a %s"),
-				*ParsedParameter.ParsedType,
-				*ShaderBindingName,
-				*CppCodeName,
-				*ExpectedShaderType);
-			GetParameterFileAndLine(ParsedParameter, Error.ErrorVirtualFilePath, Error.ErrorLineString);
-
-			CompilerOutput.Errors.Add(Error);
-			CompilerOutput.bSucceeded = false;
-
-			AddNoteToDisplayShaderParameterMemberOnCppSide(CompilerInput, ParsedParameter, CompilerOutput);
-		}
-	}
-
-	// Validate parameter size, in case this is an array.
-	if (!bBindlessHack && ReflectionSize > int32(ParsedParameter.Member->GetMemberSize()))
-	{
-		FString CppCodeName = CompilerInput.RootParametersStructure->GetFullMemberCodeName(ParsedParameter.ConstantBufferOffset);
-
-		FShaderCompilerError Error;
-		Error.StrippedErrorMessage = FString::Printf(
-			TEXT("Error: The size required to bind shader parameter %s is %i bytes, smaller than %s that is %i bytes in the parameter structure."),
-			*ShaderBindingName,
-			ReflectionSize,
-			*CppCodeName,
-			ParsedParameter.Member->GetMemberSize());
-		GetParameterFileAndLine(ParsedParameter, Error.ErrorVirtualFilePath, Error.ErrorLineString);
-
-		CompilerOutput.Errors.Add(Error);
-		CompilerOutput.bSucceeded = false;
-
-		AddNoteToDisplayShaderParameterMemberOnCppSide(CompilerInput, ParsedParameter, CompilerOutput);
-	}
-}
-
-void FShaderParameterParser::ValidateShaderParameterTypes(
-	const FShaderCompilerInput& CompilerInput,
-	bool bPlatformSupportsPrecisionModifier,
-	FShaderCompilerOutput& CompilerOutput) const
-{
-	// The shader doesn't have any parameter binding through shader structure, therefore don't do anything.
-	if (!CompilerInput.RootParametersStructure)
-	{
-		return;
-	}
-
-	if (!CompilerOutput.bSucceeded)
-	{
-		return;
-	}
-
-	const TMap<FString, FParameterAllocation>& ParametersFoundByCompiler = CompilerOutput.ParameterMap.GetParameterMap();
-
-	CompilerInput.RootParametersStructure->IterateShaderParameterMembers(
-		[&](const FShaderParametersMetadata& ParametersMetadata,
-			const FShaderParametersMetadata::FMember& Member,
-			const TCHAR* ShaderBindingName,
-			uint16 ByteOffset)
-	{
-		if (
-			Member.GetBaseType() != UBMT_INT32 &&
-			Member.GetBaseType() != UBMT_UINT32 &&
-			Member.GetBaseType() != UBMT_FLOAT32)
-		{
-			return;
-		}
-
-		const FParsedShaderParameter& ParsedParameter = ParsedParameters[ShaderBindingName];
-
-		// Did not find shader parameter in code.
-		if (!ParsedParameter.IsFound())
-		{
-			// Verify the shader compiler also did not find this parameter to make sure there is no bug in the parser.
-			checkf(
-				!ParametersFoundByCompiler.Contains(ShaderBindingName),
-				TEXT("Looks like there is a bug in FShaderParameterParser ParameterName=%s DumpDebugInfoPath=%s"),
-				ShaderBindingName,
-				*CompilerInput.DumpDebugInfoPath);
-			return;
-		}
-
-		int32 BoundOffset = 0;
-		int32 BoundSize = 0;
-		if (const FParameterAllocation* ParameterAllocation = ParametersFoundByCompiler.Find(ShaderBindingName))
-		{
-			BoundOffset = ParameterAllocation->BaseIndex;
-			BoundSize = ParameterAllocation->Size;
-		}
-
-		ValidateShaderParameterType(CompilerInput, ShaderBindingName, BoundOffset, BoundSize, bPlatformSupportsPrecisionModifier, CompilerOutput);
-	});
-}
-
-void FShaderParameterParser::ExtractFileAndLine(int32 PragamLineoffset, int32 LineOffset, FString& OutFile, FString& OutLine) const
-{
-	if (PragamLineoffset == -1)
-	{
-		return;
-	}
-
-	check(FCString::Strncmp((*OriginalParsedShader) + PragamLineoffset, TEXT("#line "), 6) == 0);
-
-	const int32 ShaderSourceLen = OriginalParsedShader.Len();
-
-	int32 StartFilePos = -1;
-	int32 EndFilePos = -1;
-	int32 StartLinePos = PragamLineoffset + 6;
-	int32 EndLinePos = -1;
-
-	for (int32 Cursor = StartLinePos; Cursor < ShaderSourceLen; Cursor++)
-	{
-		const TCHAR Char = OriginalParsedShader[Cursor];
-
-		if (Char == '\n')
-		{
-			break;
-		}
-
-		if (EndLinePos == -1)
-		{
-			if (Char > '9' || Char < '0')
-			{
-				EndLinePos = Cursor - 1;
-			}
-		}
-		else if (StartFilePos == -1)
-		{
-			if (Char == '"')
-			{
-				StartFilePos = Cursor + 1;
-			}
-		}
-		else if (EndFilePos == -1)
-		{
-			if (Char == '"')
-			{
-				EndFilePos = Cursor - 1;
-				break;
-			}
-		}
-	}
-
-	check(StartFilePos != -1);
-	check(EndFilePos != -1);
-	check(EndLinePos != -1);
-
-	OutFile = OriginalParsedShader.Mid(StartFilePos, EndFilePos - StartFilePos + 1);
-	FString LineBasis = OriginalParsedShader.Mid(StartLinePos, EndLinePos - StartLinePos + 1);
-
-	int32 FinalLine = FCString::Atoi(*LineBasis) + LineOffset;
-	OutLine = FString::FromInt(FinalLine);
-}
-
 // The cross compiler doesn't yet support struct initializers needed to construct static structs for uniform buffers
 // Replace all uniform buffer struct member references (View.WorldToClip) with a flattened name that removes the struct dependency (View_WorldToClip)
 void RemoveUniformBuffersFromSource(const FShaderCompilerEnvironment& Environment, FString& PreprocessedShaderSource)
diff --git a/Engine/Source/Developer/ShaderCompilerCommon/Private/ShaderParameterParser.cpp b/Engine/Source/Developer/ShaderCompilerCommon/Private/ShaderParameterParser.cpp
new file mode 100644
--- /dev/null
+++ b/Engine/Source/Developer/ShaderCompilerCommon/Private/ShaderParameterParser.cpp
@@ -0,0 +1,821 @@
+// Copyright Epic Games, Inc. All Rights Reserved.
+
+#include "ShaderParameterParser.h"
+#include "ShaderCompilerCommon.h"
+
+static void AddNoteToDisplayShaderParameterMemberOnCppSide(
+	const FShaderCompilerInput& CompilerInput,
+	const FShaderParameterParser::FParsedShaderParameter& ParsedParameter,
+	FShaderCompilerOutput& CompilerOutput)
+{
+	const FShaderParametersMetadata* MemberContainingStruct = nullptr;
+	const FShaderParametersMetadata::FMember* Member = nullptr;
+	{
+		int32 ArrayElementId = 0;
+		FString NamePrefix;
+		CompilerInput.RootParametersStructure->FindMemberFromOffset(ParsedParameter.ConstantBufferOffset, &MemberContainingStruct, &Member, &ArrayElementId, &NamePrefix);
+	}
+
+	FString CppCodeName = CompilerInput.RootParametersStructure->GetFullMemberCodeName(ParsedParameter.ConstantBufferOffset);
+
+	FShaderCompilerError Error;
+	Error.StrippedErrorMessage = FString::Printf(
+		TEXT("Note: Definition of %s"),
+		*CppCodeName);
+	Error.ErrorVirtualFilePath = ANSI_TO_TCHAR(MemberContainingStruct->GetFileName());
+	Error.ErrorLineString = FString::FromInt(Member->GetFileLine());
+
+	CompilerOutput.Errors.Add(Error);
+}
+
+inline bool MemberWasPotentiallyMoved(const FShaderParametersMetadata::FMember& InMember)
+{
+	const EUniformBufferBaseType BaseType = InMember.GetBaseType();
+	return BaseType == UBMT_INT32 || BaseType == UBMT_UINT32 || BaseType == UBMT_FLOAT32
+		|| BaseType == UBMT_TEXTURE || BaseType == UBMT_SRV || BaseType == UBMT_UAV
+		|| BaseType == UBMT_RDG_TEXTURE || BaseType == UBMT_RDG_TEXTURE_SRV || BaseType == UBMT_RDG_TEXTURE_UAV
+		|| BaseType == UBMT_RDG_BUFFER_SRV || BaseType == UBMT_RDG_BUFFER_UAV
+		|| BaseType == UBMT_SAMPLER
+		;
+}
+
+bool FShaderParameterParser::ParseAndMoveShaderParametersToRootConstantBuffer(
+	const FShaderCompilerInput& CompilerInput,
+	FShaderCompilerOutput& CompilerOutput,
+	FString& PreprocessedShaderSource,
+	const TCHAR* ConstantBufferType)
+{
+	// The shader doesn't have any parameter binding through shader structure, therefore don't do anything.
+	if (!CompilerInput.RootParametersStructure)
+	{
+		return true;
+	}
+
+	const bool bMoveToRootConstantBuffer = ConstantBufferType != nullptr;
+	OriginalParsedShader = PreprocessedShaderSource;
+
+	// Reserves the number of parameters up front.
+	ParsedParameters.Reserve(CompilerInput.RootParametersStructure->GetSize() / sizeof(int32));
+
+	CompilerInput.RootParametersStructure->IterateShaderParameterMembers(
+		[&](const FShaderParametersMetadata& ParametersMetadata,
+			const FShaderParametersMetadata::FMember& Member,
+			const TCHAR* ShaderBindingName,
+			uint16 ByteOffset)
+	{
+		FParsedShaderParameter ParsedParameter;
+		ParsedParameter.Member = &Member;
+		ParsedParameter.ConstantBufferOffset = ByteOffset;
+		check(ParsedParameter.IsBindable());
+
+		ParsedParameters.Add(ShaderBindingName, ParsedParameter);
+	});
+
+	bool bSuccess = true;
+
+	// Browse the code for global shader parameter, Save their type and erase them white spaces.
+	{
+		enum class EState
+		{
+			// When to look for something to scan.
+			Scanning,
+
+			// When going to next ; in the global scope and reset.
+			GoToNextSemicolonAndReset,
+
+			// Parsing what might be a type of the parameter.
+			ParsingPotentialType,
+			FinishedPotentialType,
+
+			// Parsing what might be a name of the parameter.
+			ParsingPotentialName,
+			FinishedPotentialName,
+
+			// Parsing what looks like array of the parameter.
+			ParsingPotentialArraySize,
+			FinishedArraySize,
+
+			// Found a parameter, just finish to it's semi colon.
+			FoundParameter,
+		};
+
+		const int32 ShaderSourceLen = PreprocessedShaderSource.Len();
+
+		int32 CurrentPragamLineoffset = -1;
+		int32 CurrentLineoffset = 0;
+
+		int32 TypeQualifierStartPos = -1;
+		int32 TypeStartPos = -1;
+		int32 TypeEndPos = -1;
+		int32 NameStartPos = -1;
+		int32 NameEndPos = -1;
+		int32 ArrayStartPos = -1;
+		int32 ArrayEndPos = -1;
+		int32 ScopeIndent = 0;
+
+		EState State = EState::Scanning;
+		bool bGoToNextLine = false;
+
+		auto ResetState = [&]()
+		{
+			TypeQualifierStartPos = -1;
+			TypeStartPos = -1;
+			TypeEndPos = -1;
+			NameStartPos = -1;
+			NameEndPos = -1;
+			ArrayStartPos = -1;
+			ArrayEndPos = -1;
+			State = EState::Scanning;
+		};
+
+		auto EmitError = [&](const FString& ErrorMessage)
+		{
+			FShaderCompilerError Error;
+			Error.StrippedErrorMessage = ErrorMessage;
+			ExtractFileAndLine(CurrentPragamLineoffset, CurrentLineoffset, Error.ErrorVirtualFilePath, Error.ErrorLineString);
+			CompilerOutput.Errors.Add(Error);
+			bSuccess = false;
+		};
+
+		auto EmitUnpextectedHLSLSyntaxError = [&]()
+		{
+			EmitError(TEXT("Unexpected syntax when parsing shader parameters from shader code."));
+			State = EState::GoToNextSemicolonAndReset;
+		};
+
+		for (int32 Cursor = 0; Cursor < ShaderSourceLen; Cursor++)
+		{
+			const TCHAR Char = PreprocessedShaderSource[Cursor];
+
+			auto FoundShaderParameter = [&]()
+			{
+				check(Char == ';');
+				check(TypeStartPos != -1);
+				check(TypeEndPos != -1);
+				check(NameStartPos != -1);
+				check(NameEndPos != -1);
+
+				FString Type = PreprocessedShaderSource.Mid(TypeStartPos, TypeEndPos - TypeStartPos + 1);
+				FString Name = PreprocessedShaderSource.Mid(NameStartPos, NameEndPos - NameStartPos + 1);
+				const FString ParsedName = Name;
+
+				const bool bBindlessParameter = UE::ShaderCompilerCommon::RemoveBindlessParameterPrefix(Name);
+
+				FParsedShaderParameter ParsedParameter;
+				bool bUpdateParsedParameters = false;
+				bool bEraseOriginalParameter = false;
+				if (ParsedParameters.Contains(Name))
+				{
+					if (ParsedParameters.FindChecked(Name).IsFound())
+					{
+						// If it has already been found, it means it is duplicated. Do nothing and let the shader compiler throw the error.
+					}
+					else
+					{
+						// Update the parsed parameters
+						bUpdateParsedParameters = true;
+						ParsedParameter = ParsedParameters.FindChecked(Name);
+
+						// Erase the parameter to move it into the root constant buffer.
+						if (bMoveToRootConstantBuffer && ParsedParameter.IsBindable())
+						{
+							EUniformBufferBaseType BaseType = ParsedParameter.Member->GetBaseType();
+							bEraseOriginalParameter = BaseType == UBMT_INT32 || BaseType == UBMT_UINT32 || BaseType == UBMT_FLOAT32 || bBindlessParameter;
+						}
+					}
+				}
+				else
+				{
+					// Update the parsed parameters to still have file and line number.
+					bUpdateParsedParameters = true;
+				}
+
+				// Update 
+				if (bUpdateParsedParameters)
+				{
+					ParsedParameter.ParsedName = ParsedName;
+					ParsedParameter.ParsedType = Type;
+					ParsedParameter.ParsedPragmaLineoffset = CurrentPragamLineoffset;
+					ParsedParameter.ParsedLineOffset = CurrentLineoffset;
+					ParsedParameter.bOriginalParameterErased = bEraseOriginalParameter;
+
+					if (ArrayStartPos != -1 && ArrayEndPos != -1)
+					{
+						ParsedParameter.ParsedArraySize = PreprocessedShaderSource.Mid(ArrayStartPos + 1, ArrayEndPos - ArrayStartPos - 1);
+					}
+
+					ParsedParameters.Add(Name, ParsedParameter);
+				}
+
+				// Erases this shader parameter conserving the same line numbers.
+				if (bEraseOriginalParameter)
+				{
+					for (int32 j = (TypeQualifierStartPos != -1 ? TypeQualifierStartPos : TypeStartPos); j <= Cursor; j++)
+					{
+						if (PreprocessedShaderSource[j] != '\r' && PreprocessedShaderSource[j] != '\n')
+							PreprocessedShaderSource[j] = ' ';
+					}
+				}
+
+				ResetState();
+			};
+
+			const bool bIsWhiteSpace = Char == ' ' || Char == '\t' || Char == '\r' || Char == '\n';
+			const bool bIsLetter = (Char >= 'a' && Char <= 'z') || (Char >= 'A' && Char <= 'Z');
+			const bool bIsNumber = Char >= '0' && Char <= '9';
+
+			const TCHAR* UpComing = (*PreprocessedShaderSource) + Cursor;
+			const int32 RemainingSize = ShaderSourceLen - Cursor;
+
+			CurrentLineoffset += Char == '\n';
+
+			// Go to the next line if this is a preprocessor macro.
+			if (bGoToNextLine)
+			{
+				if (Char == '\n')
+				{
+					bGoToNextLine = false;
+				}
+				continue;
+			}
+			else if (Char == '#')
+			{
+				if (RemainingSize > 6 && FCString::Strncmp(UpComing, TEXT("#line "), 6) == 0)
+				{
+					CurrentPragamLineoffset = Cursor;
+					CurrentLineoffset = -1; // that will be incremented to 0 when reaching the \n at the end of the #line
+				}
+
+				bGoToNextLine = true;
+				continue;
+			}
+
+			// If within a scope, just carry on until outside the scope.
+			if (ScopeIndent > 0 || Char == '{')
+			{
+				if (Char == '{')
+				{
+					ScopeIndent++;
+				}
+				else if (Char == '}')
+				{
+					ScopeIndent--;
+					if (ScopeIndent == 0)
+					{
+						ResetState();
+					}
+				}
+				continue;
+			}
+
+			if (State == EState::Scanning)
+			{
+				if (bIsLetter)
+				{
+					static const TCHAR* KeywordTable[] = {
+						TEXT("enum"),
+						TEXT("class"),
+						TEXT("const"),
+						TEXT("struct"),
+						TEXT("static"),
+					};
+					static int32 KeywordTableSize[] = {4, 5, 5, 6, 6};
+
+					int32 RecognisedKeywordId = -1;
+					for (int32 KeywordId = 0; KeywordId < UE_ARRAY_COUNT(KeywordTable); KeywordId++)
+					{
+						const TCHAR* Keyword = KeywordTable[KeywordId];
+						const int32 KeywordSize = KeywordTableSize[KeywordId];
+
+						if (RemainingSize > KeywordSize)
+						{
+							TCHAR KeywordEndTestChar = UpComing[KeywordSize];
+
+							if ((KeywordEndTestChar == ' ' || KeywordEndTestChar == '\r' || KeywordEndTestChar == '\n' || KeywordEndTestChar == '\t') &&
+								FCString::Strncmp(UpComing, Keyword, KeywordSize) == 0)
+							{
+								RecognisedKeywordId = KeywordId;
+								break;
+							}
+						}
+					}
+
+					if (RecognisedKeywordId == -1)
+					{
+						// Might have found beginning of the type of a parameter.
+						State = EState::ParsingPotentialType;
+						TypeStartPos = Cursor;
+					}
+					else if (RecognisedKeywordId == 2)
+					{
+						// Ignore the const keywords, but still parse given it might still be a shader parameter.
+						if (TypeQualifierStartPos == -1)
+						{
+							// If the parameter is erased, we also have to erase *all* 'const'-qualifiers, e.g. "const int Foo" or "const const int Foo".
+							TypeQualifierStartPos = Cursor;
+						}
+						Cursor += KeywordTableSize[RecognisedKeywordId];
+					}
+					else
+					{
+						// Purposefully ignore enum, class, struct, static
+						State = EState::GoToNextSemicolonAndReset;
+					}
+				}
+				else if (bIsWhiteSpace)
+				{
+					// Keep parsing void.
+				}
+				else if (Char == ';')
+				{
+					// Looks like redundant semicolon, just ignore and keep scanning.
+				}
+				else
+				{
+					// No idea what this is, just go to next semi colon.
+					State = EState::GoToNextSemicolonAndReset;
+				}
+			}
+			else if (State == EState::GoToNextSemicolonAndReset)
+			{
+				// If need to go to next global semicolon and reach it. Resume browsing.
+				if (Char == ';')
+				{
+					ResetState();
+				}
+			}
+			else if (State == EState::ParsingPotentialType)
+			{
+				// Found character legal for a type...
+				if (bIsLetter ||
+					bIsNumber ||
+					Char == '<' || Char == '>' || Char == '_')
+				{
+					// Keep browsing what might be type of the parameter.
+				}
+				else if (bIsWhiteSpace)
+				{
+					// Might have found a type.
+					State = EState::FinishedPotentialType;
+					TypeEndPos = Cursor - 1;
+				}
+				else
+				{
+					// Found unexpected character in the type.
+					State = EState::GoToNextSemicolonAndReset;
+				}
+			}
+			else if (State == EState::FinishedPotentialType)
+			{
+				if (bIsLetter)
+				{
+					// Might have found beginning of the name of a parameter.
+					State = EState::ParsingPotentialName;
+					NameStartPos = Cursor;
+				}
+				else if (bIsWhiteSpace)
+				{
+					// Keep parsing void.
+				}
+				else
+				{
+					// No idea what this is, just go to next semi colon.
+					State = EState::GoToNextSemicolonAndReset;
+				}
+			}
+			else if (State == EState::ParsingPotentialName)
+			{
+				// Found character legal for a name...
+				if (bIsLetter ||
+					bIsNumber ||
+					Char == '_')
+				{
+					// Keep browsing what might be name of the parameter.
+				}
+				else if (Char == ':' || Char == '=')
+				{
+					// Found a parameter with syntax:
+					// uint MyParameter : <whatever>;
+					// uint MyParameter = <DefaultValue>;
+					NameEndPos = Cursor - 1;
+					State = EState::FoundParameter;
+				}
+				else if (Char == ';')
+				{
+					// Found a parameter with syntax:
+					// uint MyParameter;
+					NameEndPos = Cursor - 1;
+					FoundShaderParameter();
+				}
+				else if (Char == '[')
+				{
+					// Syntax:
+					//  uint MyArray[
+					NameEndPos = Cursor - 1;
+					ArrayStartPos = Cursor;
+					State = EState::ParsingPotentialArraySize;
+				}
+				else if (bIsWhiteSpace)
+				{
+					// Might have found a name.
+					// uint MyParameter <Still need to know what is after>;
+					NameEndPos = Cursor - 1;
+					State = EState::FinishedPotentialName;
+				}
+				else
+				{
+					// Found unexpected character in the name.
+					// syntax:
+					// uint MyFunction(<Don't care what is after>
+					State = EState::GoToNextSemicolonAndReset;
+				}
+			}
+			else if (
+				State == EState::FinishedPotentialName ||
+				State == EState::FinishedArraySize)
+			{
+				if (Char == ';')
+				{
+					// Found a parameter with syntax:
+					// uint MyParameter <a bit of OK stuf>;
+					FoundShaderParameter();
+				}
+				else if (Char == ':')
+				{
+					// Found a parameter with syntax:
+					// uint MyParameter <a bit of OK stuf> : <Ignore all this crap>;
+					State = EState::FoundParameter;
+				}
+				else if (Char == '=')
+				{
+					// Found syntax that doesn't make any sens:
+					// uint MyParameter <a bit of OK stuf> = <Ignore all this crap>;
+					State = EState::FoundParameter;
+					// TDOO: should error out that this is useless.
+				}
+				else if (Char == '[')
+				{
+					if (State == EState::FinishedPotentialName)
+					{
+						// Syntax:
+						//  uint MyArray [
+						ArrayStartPos = Cursor;
+						State = EState::ParsingPotentialArraySize;
+					}
+					else
+					{
+						EmitError(TEXT("Shader parameters can only support one dimensional array"));
+					}
+				}
+				else if (bIsWhiteSpace)
+				{
+					// Keep parsing void.
+				}
+				else
+				{
+					// Found unexpected stuff.
+					State = EState::GoToNextSemicolonAndReset;
+				}
+			}
+			else if (State == EState::ParsingPotentialArraySize)
+			{
+				if (Char == ']')
+				{
+					ArrayEndPos = Cursor;
+					State = EState::FinishedArraySize;
+				}
+				else if (Char == ';')
+				{
+					EmitUnpextectedHLSLSyntaxError();
+				}
+				else
+				{
+					// Keep going through the array size that might be a complex expression.
+				}
+			}
+			else if (State == EState::FoundParameter)
+			{
+				if (Char == ';')
+				{
+					FoundShaderParameter();
+				}
+				else
+				{
+					// Cary on skipping all crap we don't care about shader parameter until we find it's semi colon.
+				}
+			}
+			else
+			{
+				unimplemented();
+			}
+		} // for (int32 Cursor = 0; Cursor < PreprocessedShaderSource.Len(); Cursor++)
+	}
+
+	// Generate the root cbuffer content.
+	if (bMoveToRootConstantBuffer)
+	{
+		FString RootCBufferContent;
+
+		CompilerInput.RootParametersStructure->IterateShaderParameterMembers(
+			[&](const FShaderParametersMetadata& ParametersMetadata,
+				const FShaderParametersMetadata::FMember& Member,
+				const TCHAR* ShaderBindingName,
+				uint16 ByteOffset)
+		{
+			if (MemberWasPotentiallyMoved(Member))
+			{
+				FParsedShaderParameter* ParsedParameter = ParsedParameters.Find(ShaderBindingName);
+				if (ParsedParameter && ParsedParameter->IsFound() && ParsedParameter->bOriginalParameterErased)
+				{
+					uint32 ConstantRegister = ByteOffset / 16;
+					const TCHAR* ConstantSwizzle = [ByteOffset]()
+						{
+							switch (ByteOffset % 16)
+							{
+							default: unimplemented();
+							case 0:  return TEXT("");
+							case 4:  return TEXT(".y");
+							case 8:  return TEXT(".z");
+							case 12: return TEXT(".w");
+							}
+						}();
+
+					if (!ParsedParameter->ParsedArraySize.IsEmpty())
+					{
+						RootCBufferContent.Append(FString::Printf(
+							TEXT("%s %s[%s] : packoffset(c%d%s);\r\n"),
+							*ParsedParameter->ParsedType,
+							*ParsedParameter->ParsedName,
+							*ParsedParameter->ParsedArraySize,
+							ConstantRegister,
+							ConstantSwizzle));
+					}
+					else
+					{
+						RootCBufferContent.Append(FString::Printf(
+							TEXT("%s %s : packoffset(c%d%s);\r\n"),
+							*ParsedParameter->ParsedType,
+							*ParsedParameter->ParsedName,
+							ConstantRegister,
+							ConstantSwizzle));
+					}
+				}
+			}
+		});
+
+		FString CBufferCodeBlock = FString::Printf(
+			TEXT("%s %s\r\n")
+			TEXT("{\r\n")
+			TEXT("%s")
+			TEXT("}\r\n\r\n"),
+			ConstantBufferType,
+			FShaderParametersMetadata::kRootUniformBufferBindingName,
+			*RootCBufferContent);
+
+		FString NewShaderCode = (
+			MakeInjectedShaderCodeBlock(TEXT("ParseAndMoveShaderParametersToRootConstantBuffer"), CBufferCodeBlock) +
+			PreprocessedShaderSource);
+
+		PreprocessedShaderSource = MoveTemp(NewShaderCode);
+
+		bMovedLoosedParametersToRootConstantBuffer = true;
+	} // if (bMoveToRootConstantBuffer)
+
+	return bSuccess;
+}
+
+void FShaderParameterParser::ValidateShaderParameterType(
+	const FShaderCompilerInput& CompilerInput,
+	const FString& ShaderBindingName,
+	int32 ReflectionOffset,
+	int32 ReflectionSize,
+	bool bPlatformSupportsPrecisionModifier,
+	FShaderCompilerOutput& CompilerOutput) const
+{
+	FString BindingName(ShaderBindingName);
+
+	const bool bBindlessHack = UE::ShaderCompilerCommon::RemoveBindlessParameterPrefix(BindingName);
+
+	const FShaderParameterParser::FParsedShaderParameter& ParsedParameter = FindParameterInfos(BindingName);
+
+	check(ParsedParameter.IsFound());
+	check(CompilerInput.RootParametersStructure);
+
+	if (ReflectionSize > 0 && bMovedLoosedParametersToRootConstantBuffer)
+	{
+		// Verify the offset of the parameter coming from shader reflections honor the packoffset()
+		check(ReflectionOffset == ParsedParameter.ConstantBufferOffset);
+	}
+
+	// Validate the shader type.
+	if (!bBindlessHack)
+	{
+		FString ExpectedShaderType;
+		ParsedParameter.Member->GenerateShaderParameterType(ExpectedShaderType, bPlatformSupportsPrecisionModifier);
+
+		const bool bShouldBeInt = ParsedParameter.Member->GetBaseType() == UBMT_INT32;
+		const bool bShouldBeUint = ParsedParameter.Member->GetBaseType() == UBMT_UINT32;
+
+		// Match parsed type with expected shader type
+		bool bIsTypeCorrect = ParsedParameter.ParsedType == ExpectedShaderType;
+
+		if (!bIsTypeCorrect)
+		{
+			// Accept half-precision floats when single-precision was requested
+			if (ParsedParameter.ParsedType.StartsWith(TEXT("half")) && ParsedParameter.Member->GetBaseType() == UBMT_FLOAT32)
+			{
+				bIsTypeCorrect = (FCString::Strcmp(*ParsedParameter.ParsedType + 4, *ExpectedShaderType + 5) == 0);
+			}
+			// Accept single-precision floats when half-precision was expected
+			else if (ParsedParameter.ParsedType.StartsWith(TEXT("float")) && ExpectedShaderType.StartsWith(TEXT("half")))
+			{
+				bIsTypeCorrect = (FCString::Strcmp(*ParsedParameter.ParsedType + 5, *ExpectedShaderType + 4) == 0);
+			}
+			// support for min16float
+			else if (ParsedParameter.ParsedType.StartsWith(TEXT("min16float")) && ExpectedShaderType.StartsWith(TEXT("float")))
+			{
+				bIsTypeCorrect = (FCString::Strcmp(*ParsedParameter.ParsedType + 10, *ExpectedShaderType + 5) == 0);
+			}
+			else if (ParsedParameter.ParsedType.StartsWith(TEXT("min16float")) && ExpectedShaderType.StartsWith(TEXT("half")))
+			{
+				bIsTypeCorrect = (FCString::Strcmp(*ParsedParameter.ParsedType + 10, *ExpectedShaderType + 4) == 0);
+			}
+		}
+
+		// Allow silent casting between signed and unsigned on shader bindings.
+		if (!bIsTypeCorrect && (bShouldBeInt || bShouldBeUint))
+		{
+			FString NewExpectedShaderType;
+			if (bShouldBeInt)
+			{
+				// tries up with an uint.
+				NewExpectedShaderType = TEXT("u") + ExpectedShaderType;
+			}
+			else
+			{
+				// tries up with an int.
+				NewExpectedShaderType = ExpectedShaderType;
+				NewExpectedShaderType.RemoveAt(0);
+			}
+
+			bIsTypeCorrect = ParsedParameter.ParsedType == NewExpectedShaderType;
+		}
+
+		if (!bIsTypeCorrect)
+		{
+			FString CppCodeName = CompilerInput.RootParametersStructure->GetFullMemberCodeName(ParsedParameter.ConstantBufferOffset);
+
+			FShaderCompilerError Error;
+			Error.StrippedErrorMessage = FString::Printf(
+				TEXT("Error: Type %s of shader parameter %s in shader mismatch the shader parameter structure: %s expects a %s"),
+				*ParsedParameter.ParsedType,
+				*ShaderBindingName,
+				*CppCodeName,
+				*ExpectedShaderType);
+			GetParameterFileAndLine(ParsedParameter, Error.ErrorVirtualFilePath, Error.ErrorLineString);
+
+			CompilerOutput.Errors.Add(Error);
+			CompilerOutput.bSucceeded = false;
+
+			AddNoteToDisplayShaderParameterMemberOnCppSide(CompilerInput, ParsedParameter, CompilerOutput);
+		}
+	}
+
+	// Validate parameter size, in case this is an array.
+	if (!bBindlessHack && ReflectionSize > int32(ParsedParameter.Member->GetMemberSize()))
+	{
+		FString CppCodeName = CompilerInput.RootParametersStructure->GetFullMemberCodeName(ParsedParameter.ConstantBufferOffset);
+
+		FShaderCompilerError Error;
+		Error.StrippedErrorMessage = FString::Printf(
+			TEXT("Error: The size required to bind shader parameter %s is %i bytes, smaller than %s that is %i bytes in the parameter structure."),
+			*ShaderBindingName,
+			ReflectionSize,
+			*CppCodeName,
+			ParsedParameter.Member->GetMemberSize());
+		GetParameterFileAndLine(ParsedParameter, Error.ErrorVirtualFilePath, Error.ErrorLineString);
+
+		CompilerOutput.Errors.Add(Error);
+		CompilerOutput.bSucceeded = false;
+
+		AddNoteToDisplayShaderParameterMemberOnCppSide(CompilerInput, ParsedParameter, CompilerOutput);
+	}
+}
+
+void FShaderParameterParser::ValidateShaderParameterTypes(
+	const FShaderCompilerInput& CompilerInput,
+	bool bPlatformSupportsPrecisionModifier,
+	FShaderCompilerOutput& CompilerOutput) const
+{
+	// The shader doesn't have any parameter binding through shader structure, therefore don't do anything.
+	if (!CompilerInput.RootParametersStructure)
+	{
+		return;
+	}
+
+	if (!CompilerOutput.bSucceeded)
+	{
+		return;
+	}
+
+	const TMap<FString, FParameterAllocation>& ParametersFoundByCompiler = CompilerOutput.ParameterMap.GetParameterMap();
+
+	CompilerInput.RootParametersStructure->IterateShaderParameterMembers(
+		[&](const FShaderParametersMetadata& ParametersMetadata,
+			const FShaderParametersMetadata::FMember& Member,
+			const TCHAR* ShaderBindingName,
+			uint16 ByteOffset)
+		{
+			if (
+				Member.GetBaseType() != UBMT_INT32 &&
+				Member.GetBaseType() != UBMT_UINT32 &&
+				Member.GetBaseType() != UBMT_FLOAT32)
+			{
+				return;
+			}
+
+			const FParsedShaderParameter& ParsedParameter = ParsedParameters[ShaderBindingName];
+
+			// Did not find shader parameter in code.
+			if (!ParsedParameter.IsFound())
+			{
+				// Verify the shader compiler also did not find this parameter to make sure there is no bug in the parser.
+				checkf(
+					!ParametersFoundByCompiler.Contains(ShaderBindingName),
+					TEXT("Looks like there is a bug in FShaderParameterParser ParameterName=%s DumpDebugInfoPath=%s"),
+					ShaderBindingName,
+					*CompilerInput.DumpDebugInfoPath);
+				return;
+			}
+
+			int32 BoundOffset = 0;
+			int32 BoundSize = 0;
+			if (const FParameterAllocation* ParameterAllocation = ParametersFoundByCompiler.Find(ShaderBindingName))
+			{
+				BoundOffset = ParameterAllocation->BaseIndex;
+				BoundSize = ParameterAllocation->Size;
+			}
+
+			ValidateShaderParameterType(CompilerInput, ShaderBindingName, BoundOffset, BoundSize, bPlatformSupportsPrecisionModifier, CompilerOutput);
+		});
+}
+
+void FShaderParameterParser::ExtractFileAndLine(int32 PragamLineoffset, int32 LineOffset, FString& OutFile, FString& OutLine) const
+{
+	if (PragamLineoffset == -1)
+	{
+		return;
+	}
+
+	check(FCString::Strncmp((*OriginalParsedShader) + PragamLineoffset, TEXT("#line "), 6) == 0);
+
+	const int32 ShaderSourceLen = OriginalParsedShader.Len();
+
+	int32 StartFilePos = -1;
+	int32 EndFilePos = -1;
+	int32 StartLinePos = PragamLineoffset + 6;
+	int32 EndLinePos = -1;
+
+	for (int32 Cursor = StartLinePos; Cursor < ShaderSourceLen; Cursor++)
+	{
+		const TCHAR Char = OriginalParsedShader[Cursor];
+
+		if (Char == '\n')
+		{
+			break;
+		}
+
+		if (EndLinePos == -1)
+		{
+			if (Char > '9' || Char < '0')
+			{
+				EndLinePos = Cursor - 1;
+			}
+		}
+		else if (StartFilePos == -1)
+		{
+			if (Char == '"')
+			{
+				StartFilePos = Cursor + 1;
+			}
+		}
+		else if (EndFilePos == -1)
+		{
+			if (Char == '"')
+			{
+				EndFilePos = Cursor - 1;
+				break;
+			}
+		}
+	}
+
+	check(StartFilePos != -1);
+	check(EndFilePos != -1);
+	check(EndLinePos != -1);
+
+	OutFile = OriginalParsedShader.Mid(StartFilePos, EndFilePos - StartFilePos + 1);
+	FString LineBasis = OriginalParsedShader.Mid(StartLinePos, EndLinePos - StartLinePos + 1);
+
+	int32 FinalLine = FCString::Atoi(*LineBasis) + LineOffset;
+	OutLine = FString::FromInt(FinalLine);
+}
diff --git a/Engine/Source/Developer/ShaderCompilerCommon/Public/ShaderCompilerCommon.h b/Engine/Source/Developer/ShaderCompilerCommon/Public/ShaderCompilerCommon.h
--- a/Engine/Source/Developer/ShaderCompilerCommon/Public/ShaderCompilerCommon.h
+++ b/Engine/Source/Developer/ShaderCompilerCommon/Public/ShaderCompilerCommon.h
@@ -7,6 +7,11 @@
 #include "ShaderCompilerCore.h"
 #include "CrossCompilerDefinitions.h"
 #include "ShaderConductorContext.h"
+#if UE_ENABLE_INCLUDE_ORDER_DEPRECATED_IN_5_1
+#include "ShaderParameterParser.h"
+#endif
+
+class FShaderParameterParser;
 
 namespace UE::ShaderCompilerCommon
 {
@@ -49,107 +54,6 @@ extern SHADERCOMPILERCOMMON_API void BuildResourceTableTokenStream(
 // Finds the number of used uniform buffers in a resource map
 extern SHADERCOMPILERCOMMON_API int16 GetNumUniformBuffersUsed(const FShaderCompilerResourceTable& InSRT);
 
-/** Validates and moves all the shader loose data parameter defined in the root scope of the shader into the root uniform buffer. */
-class SHADERCOMPILERCOMMON_API FShaderParameterParser
-{
-public:
-	struct FParsedShaderParameter
-	{
-	public:
-		/** Original information about the member. */
-		const FShaderParametersMetadata::FMember* Member = nullptr;
-
-		/** Information found about the member when parsing the preprocessed code. */
-		FString ParsedName;
-		FString ParsedType;
-		FString ParsedArraySize;
-
-		/** Offset the member should be in the constant buffer. */
-		int32 ConstantBufferOffset = 0;
-
-		/** Records if the parameter was fully moved. */
-		bool bOriginalParameterErased{};
-
-		/* Returns whether the shader parameter has been found when parsing. */
-		bool IsFound() const
-		{
-			return !ParsedType.IsEmpty();
-		}
-
-		/** Returns whether the shader parameter is bindable to the shader parameter structure. */
-		bool IsBindable() const
-		{
-			return Member != nullptr;
-		}
-
-	private:
-		int32 ParsedPragmaLineoffset = 0;
-		int32 ParsedLineOffset = 0;
-
-		friend class FShaderParameterParser;
-	};
-
-	/** Parses the preprocessed shader code and move the parameters into root constant buffer */
-	bool ParseAndMoveShaderParametersToRootConstantBuffer(
-		const FShaderCompilerInput& CompilerInput,
-		FShaderCompilerOutput& CompilerOutput,
-		FString& PreprocessedShaderSource,
-		const TCHAR* ConstantBufferType);
-
-	/** Gets parsing information from a parameter binding name. */
-	const FParsedShaderParameter& FindParameterInfos(const FString& ParameterName) const
-	{
-		return ParsedParameters.FindChecked(ParameterName);
-	}
-
-	/** Validates the shader parameter in code is compatible with the shader parameter structure. */
-	void ValidateShaderParameterType(
-		const FShaderCompilerInput& CompilerInput,
-		const FString& ShaderBindingName,
-		int32 ReflectionOffset,
-		int32 ReflectionSize,
-		bool bPlatformSupportsPrecisionModifier,
-		FShaderCompilerOutput& CompilerOutput) const;
-
-	void ValidateShaderParameterType(
-		const FShaderCompilerInput& CompilerInput,
-		const FString& ShaderBindingName,
-		int32 ReflectionOffset,
-		int32 ReflectionSize,
-		FShaderCompilerOutput& CompilerOutput) const
-	{
-		ValidateShaderParameterType(CompilerInput, ShaderBindingName, ReflectionOffset, ReflectionSize, false, CompilerOutput);
-	}
-
-	/** Validates shader parameter map is compatible with the shader parameter structure. */
-	void ValidateShaderParameterTypes(
-		const FShaderCompilerInput& CompilerInput,
-		bool bPlatformSupportsPrecisionModifier,
-		FShaderCompilerOutput& CompilerOutput) const;
-
-	void ValidateShaderParameterTypes(
-		const FShaderCompilerInput& CompilerInput,
-		FShaderCompilerOutput& CompilerOutput) const
-	{
-		ValidateShaderParameterTypes(CompilerInput, false, CompilerOutput);
-	}
-
-	/** Gets file and line of the parameter in the shader source code. */
-	void GetParameterFileAndLine(const FParsedShaderParameter& ParsedParameter, FString& OutFile, FString& OutLine) const
-	{
-		return ExtractFileAndLine(ParsedParameter.ParsedPragmaLineoffset, ParsedParameter.ParsedLineOffset, OutFile, OutLine);
-	}
-
-private:
-	void ExtractFileAndLine(int32 PragamLineoffset, int32 LineOffset, FString& OutFile, FString& OutLine) const;
-
-	FString OriginalParsedShader;
-
-	TMap<FString, FParsedShaderParameter> ParsedParameters;
-
-	bool bMovedLoosedParametersToRootConstantBuffer = false;
-};
-
 namespace UE::ShaderCompilerCommon
 {
 	extern SHADERCOMPILERCOMMON_API FString              RemoveConstantBufferPrefix(const FString& InName);
@@ -256,12 +160,6 @@ extern SHADERCOMPILERCOMMON_API void AddNoteToDisplayShaderParameterStructureOnC
 	const FShaderParametersMetadata* ParametersStructure,
 	FShaderCompilerOutput& CompilerOutput);
 
-/** Adds a note to CompilerOutput.Error about where the shader parameter is on C++ side. */
-extern SHADERCOMPILERCOMMON_API void AddNoteToDisplayShaderParameterMemberOnCppSide(
-	const FShaderCompilerInput& CompilerInput,
-	const FShaderParameterParser::FParsedShaderParameter& ParsedParameter,
-	FShaderCompilerOutput& CompilerOutput);
-
 /** Adds an error to CompilerOutput.Error about a shader parameters that could not be bound. */
 extern SHADERCOMPILERCOMMON_API void AddUnboundShaderParameterError(
 	const FShaderCompilerInput& CompilerInput,
diff --git a/Engine/Source/Developer/ShaderCompilerCommon/Public/ShaderParameterParser.h b/Engine/Source/Developer/ShaderCompilerCommon/Public/ShaderParameterParser.h
new file mode 100644
--- /dev/null
+++ b/Engine/Source/Developer/ShaderCompilerCommon/Public/ShaderParameterParser.h
@@ -0,0 +1,110 @@
+// Copyright Epic Games, Inc. All Rights Reserved.
+
+#pragma once
+
+#include "CoreMinimal.h"
+#include "ShaderParameterMetadata.h"
+
+struct FShaderCompilerInput;
+struct FShaderCompilerOutput;
+
+/** Validates and moves all the shader loose data parameter defined in the root scope of the shader into the root uniform buffer. */
+class SHADERCOMPILERCOMMON_API FShaderParameterParser
+{
+public:
+	struct FParsedShaderParameter
+	{
+	public:
+		/** Original information about the member. */
+		const FShaderParametersMetadata::FMember* Member = nullptr;
+
+		/** Information found about the member when parsing the preprocessed code. */
+		FString ParsedName;
+		FString ParsedType;
+		FString ParsedArraySize;
+
+		/** Offset the member should be in the constant buffer. */
+		int32 ConstantBufferOffset = 0;
+
+		/** Records if the parameter was fully moved. */
+		bool bOriginalParameterErased{};
+
+		/* Returns whether the shader parameter has been found when parsing. */
+		bool IsFound() const
+		{
+			return !ParsedType.IsEmpty();
+		}
+
+		/** Returns whether the shader parameter is bindable to the shader parameter structure. */
+		bool IsBindable() const
+		{
+			return Member != nullptr;
+		}
+
+	private:
+		int32 ParsedPragmaLineoffset = 0;
+		int32 ParsedLineOffset = 0;
+
+		friend class FShaderParameterParser;
+	};
+
+	/** Parses the preprocessed shader code and move the parameters into root constant buffer */
+	bool ParseAndMoveShaderParametersToRootConstantBuffer(
+		const FShaderCompilerInput& CompilerInput,
+		FShaderCompilerOutput& CompilerOutput,
+		FString& PreprocessedShaderSource,
+		const TCHAR* ConstantBufferType);
+
+	/** Gets parsing information from a parameter binding name. */
+	const FParsedShaderParameter& FindParameterInfos(const FString& ParameterName) const
+	{
+		return ParsedParameters.FindChecked(ParameterName);
+	}
+
+	/** Validates the shader parameter in code is compatible with the shader parameter structure. */
+	void ValidateShaderParameterType(
+		const FShaderCompilerInput& CompilerInput,
+		const FString& ShaderBindingName,
+		int32 ReflectionOffset,
+		int32 ReflectionSize,
+		bool bPlatformSupportsPrecisionModifier,
+		FShaderCompilerOutput& CompilerOutput) const;
+
+	void ValidateShaderParameterType(
+		const FShaderCompilerInput& CompilerInput,
+		const FString& ShaderBindingName,
+		int32 ReflectionOffset,
+		int32 ReflectionSize,
+		FShaderCompilerOutput& CompilerOutput) const
+	{
+		ValidateShaderParameterType(CompilerInput, ShaderBindingName, ReflectionOffset, ReflectionSize, false, CompilerOutput);
+	}
+
+	/** Validates shader parameter map is compatible with the shader parameter structure. */
+	void ValidateShaderParameterTypes(
+		const FShaderCompilerInput& CompilerInput,
+		bool bPlatformSupportsPrecisionModifier,
+		FShaderCompilerOutput& CompilerOutput) const;
+
+	void ValidateShaderParameterTypes(
+		const FShaderCompilerInput& CompilerInput,
+		FShaderCompilerOutput& CompilerOutput) const
+	{
+		ValidateShaderParameterTypes(CompilerInput, false, CompilerOutput);
+	}
+
+	/** Gets file and line of the parameter in the shader source code. */
+	void GetParameterFileAndLine(const FParsedShaderParameter& ParsedParameter, FString& OutFile, FString& OutLine) const
+	{
+		return ExtractFileAndLine(ParsedParameter.ParsedPragmaLineoffset, ParsedParameter.ParsedLineOffset, OutFile, OutLine);
+	}
+
+private:
+	void ExtractFileAndLine(int32 PragamLineoffset, int32 LineOffset, FString& OutFile, FString& OutLine) const;
+
+	FString OriginalParsedShader;
+
+	TMap<FString, FParsedShaderParameter> ParsedParameters;
+
+	bool bMovedLoosedParametersToRootConstantBuffer = false;
+};
diff --git a/Engine/Source/Developer/ShaderFormatOpenGL/Private/OpenGLShaderCompiler.cpp b/Engine/Source/Developer/ShaderFormatOpenGL/Private/OpenGLShaderCompiler.cpp
--- a/Engine/Source/Developer/ShaderFormatOpenGL/Private/OpenGLShaderCompiler.cpp
+++ b/Engine/Source/Developer/ShaderFormatOpenGL/Private/OpenGLShaderCompiler.cpp
@@ -9,6 +9,7 @@
 #include "ShaderFormatOpenGL.h"
 #include "HlslccHeaderWriter.h"
 #include "SpirvReflectCommon.h"
+#include "ShaderParameterParser.h"
 #include <algorithm>
 #include <regex>
 
diff --git a/Engine/Source/Developer/VulkanShaderFormat/Private/VulkanShaderCompiler.cpp b/Engine/Source/Developer/VulkanShaderFormat/Private/VulkanShaderCompiler.cpp
--- a/Engine/Source/Developer/VulkanShaderFormat/Private/VulkanShaderCompiler.cpp
+++ b/Engine/Source/Developer/VulkanShaderFormat/Private/VulkanShaderCompiler.cpp
@@ -5,6 +5,7 @@
 #include "VulkanCommon.h"
 #include "ShaderPreprocessor.h"
 #include "ShaderCompilerCommon.h"
+#include "ShaderParameterParser.h"
 #include "HlslccHeaderWriter.h"
 #include "hlslcc.h"
 #include "SpirvReflectCommon.h"
diff --git a/Engine/Source/Developer/Windows/ShaderFormatD3D/Private/D3DShaderCompiler.cpp b/Engine/Source/Developer/Windows/ShaderFormatD3D/Private/D3DShaderCompiler.cpp
--- a/Engine/Source/Developer/Windows/ShaderFormatD3D/Private/D3DShaderCompiler.cpp
+++ b/Engine/Source/Developer/Windows/ShaderFormatD3D/Private/D3DShaderCompiler.cpp
@@ -3,6 +3,7 @@
 #include "ShaderFormatD3D.h"
 #include "ShaderPreprocessor.h"
 #include "ShaderCompilerCommon.h"
+#include "ShaderParameterParser.h"
 #include "D3D11ShaderResources.h"
 #include "D3D12RHI.h"
 #include "Misc/Paths.h"
diff --git a/Engine/Source/Developer/Windows/ShaderFormatD3D/Private/D3DShaderCompilerDXC.cpp b/Engine/Source/Developer/Windows/ShaderFormatD3D/Private/D3DShaderCompilerDXC.cpp
--- a/Engine/Source/Developer/Windows/ShaderFormatD3D/Private/D3DShaderCompilerDXC.cpp
+++ b/Engine/Source/Developer/Windows/ShaderFormatD3D/Private/D3DShaderCompilerDXC.cpp
@@ -3,6 +3,7 @@
 #include "ShaderFormatD3D.h"
 #include "ShaderPreprocessor.h"
 #include "ShaderCompilerCommon.h"
+#include "ShaderParameterParser.h"
 #include "D3D11ShaderResources.h"
 #include "D3D12RHI.h"
 #include "Misc/Paths.h"
