Fix false-positive error when adding an empty TEXT("") macro to a shader.

diff --git a/Engine/Source/Developer/ShaderCompilerCommon/Private/ShaderCompilerCommon.cpp b/Engine/Source/Developer/ShaderCompilerCommon/Private/ShaderCompilerCommon.cpp
--- a/Engine/Source/Developer/ShaderCompilerCommon/Private/ShaderCompilerCommon.cpp
+++ b/Engine/Source/Developer/ShaderCompilerCommon/Private/ShaderCompilerCommon.cpp
@@ -860,7 +860,7 @@ void TransformStringIntoCharacterArray(FString& PreprocessedShaderSource)
 	// 5. Write the function for fetching character for a given entry index
 	const uint32 EntryCount = Entries.Num();
 	FString TextChars;
-	if (EntryCount)
+	if (EntryCount>0 && GlobalCount>0)
 	{
 		// 1. Encoded character for each text entry within a single global char array
 		TextChars = FString::Printf(TEXT("static const uint TEXT_CHARS[%d] = {\n"), GlobalCount);
