Reworking the parameter rewrite code to do automatic bindless
- Refactoring FShaderParameterParser to use more string views since we can rely on OriginalParsedShader not changing.
- Split FShaderParameterParser::ParseAndMoveShaderParametersToRootConstantBuffer into smaller functions
- Moving FShaderParametersMetadata::IterateShaderParameterMembers into the only module it's used in, ShaderParameterParser.cpp
- Adding CFLAG_BindlessResources which will be set when ENABLE_BINDLESS_RESOURCES is set
- Adding CFLAG_BindlessSamplers which will be set when ENABLE_BINDLESS_SAMPLERS is set
- Adding resource/sampler detection and conversion to bindless in FShaderParameterParser

diff --git a/Engine/Shaders/Public/BindlessResources.ush b/Engine/Shaders/Public/BindlessResources.ush
--- a/Engine/Shaders/Public/BindlessResources.ush
+++ b/Engine/Shaders/Public/BindlessResources.ush
@@ -20,18 +20,21 @@
 //
 
 #if ENABLE_BINDLESS_RESOURCES && PLATFORM_SUPPORTS_BINDLESS
-    #define DEFINE_RESOURCE_VARIABLE(Type, Name) \
-        Type GetBindlessResource##Name() \
-        { \
-            return GetResourceFromHeap(Type, BindlessResource_##Name); \
-        }
+	#define DEFINE_RESOURCE_VARIABLE(Type, Name) \
+		Type GetBindlessResource##Name() \
+		{ \
+			return GetResourceFromHeap(Type, BindlessResource_##Name); \
+		}
 
-    #define GET_RESOURCE(Name) GetBindlessResource##Name()
+	#define GET_RESOURCE(Name) GetBindlessResource##Name()
 
-    #define DEFINE_RESOURCE(Type, Name) \
+	#define DEFINE_RESOURCE(Type, Name) \
 		DEFINE_RESOURCE_INDEX(Name); \
-        DEFINE_RESOURCE_VARIABLE(Type, Name); \
+		DEFINE_RESOURCE_VARIABLE(Type, Name); \
 		static const Type Name = GET_RESOURCE(Name);
+
+	COMPILER_DEFINE AUTO_BINDLESS_RESOURCE_INDEX(Name)          DEFINE_RESOURCE_INDEX(Name)
+	COMPILER_DEFINE AUTO_BINDLESS_RESOURCE_VARIABLE(Type, Name) DEFINE_RESOURCE_VARIABLE(Type, Name); static const Type Name = GET_RESOURCE(Name)
 #else
 	#define DEFINE_RESOURCE_VARIABLE(Type, Name)  Type Name
 	#define DEFINE_RESOURCE(Type, Name)           Type Name
@@ -41,16 +44,19 @@
 #if ENABLE_BINDLESS_SAMPLERS && PLATFORM_SUPPORTS_BINDLESS
 	#define DEFINE_SAMPLER_VARIABLE(Type, Name) \
 		Type GetBindlessSampler##Name() \
-        { \
-            return GetSamplerFromHeap(Type, BindlessSampler_##Name); \
-        }
+		{ \
+			return GetSamplerFromHeap(Type, BindlessSampler_##Name); \
+		}
 
-    #define GET_SAMPLER(Name)  GetBindlessSampler##Name()
+	#define GET_SAMPLER(Name)  GetBindlessSampler##Name()
 
-    #define DEFINE_SAMPLER(Type, Name) \
-        DEFINE_SAMPLER_INDEX(Name); \
-        DEFINE_SAMPLER_VARIABLE(Type, Name); \
+	#define DEFINE_SAMPLER(Type, Name) \
+		DEFINE_SAMPLER_INDEX(Name); \
+		DEFINE_SAMPLER_VARIABLE(Type, Name); \
 		static const Type Name = GET_SAMPLER(Name);
+
+	COMPILER_DEFINE AUTO_BINDLESS_SAMPLER_INDEX(Name)          uint BindlessSampler_##Name
+	COMPILER_DEFINE AUTO_BINDLESS_SAMPLER_VARIABLE(Type, Name) Type GetBindlessSampler##Name() { return GetSamplerFromHeap(Type, BindlessSampler_##Name); }; static const Type Name = GetBindlessSampler##Name()
 #else
 	#define DEFINE_SAMPLER_VARIABLE(Type, Name)  Type Name
 	#define DEFINE_SAMPLER(Type, Name)           Type Name
diff --git a/Engine/Shaders/Public/Platform/D3D/D3DCommon.ush b/Engine/Shaders/Public/Platform/D3D/D3DCommon.ush
--- a/Engine/Shaders/Public/Platform/D3D/D3DCommon.ush
+++ b/Engine/Shaders/Public/Platform/D3D/D3DCommon.ush
@@ -119,6 +119,6 @@ void UEReportAssertWithPayload(uint AssertID, uint4 Payload)
 
 #if PLATFORM_SUPPORTS_BINDLESS
 	// SM6.6 path
-	#define GetResourceFromHeap(Type, Index) ResourceDescriptorHeap[Index]
-	#define GetSamplerFromHeap(Type, Index) SamplerDescriptorHeap[Index]
+	COMPILER_DEFINE GetResourceFromHeap(Type, Index) ResourceDescriptorHeap[Index]
+	COMPILER_DEFINE GetSamplerFromHeap(Type, Index) SamplerDescriptorHeap[Index]
 #endif
diff --git a/Engine/Source/Developer/Apple/MetalShaderFormat/Private/MetalShaderCompiler.cpp b/Engine/Source/Developer/Apple/MetalShaderFormat/Private/MetalShaderCompiler.cpp
--- a/Engine/Source/Developer/Apple/MetalShaderFormat/Private/MetalShaderCompiler.cpp
+++ b/Engine/Source/Developer/Apple/MetalShaderFormat/Private/MetalShaderCompiler.cpp
@@ -1059,8 +1059,7 @@ void CompileShader_Metal(const FShaderCompilerInput& _Input,FShaderCompilerOutpu
 	}
 
 	FShaderParameterParser ShaderParameterParser;
-	if (!ShaderParameterParser.ParseAndMoveShaderParametersToRootConstantBuffer(
-		Input, Output, PreprocessedShader, /* ConstantBufferType = */ nullptr))
+	if (!ShaderParameterParser.ParseAndModify(Input, Output, PreprocessedShader, /* ConstantBufferType = */ nullptr))
 	{
 		// The FShaderParameterParser will add any relevant errors.
 		return;
diff --git a/Engine/Source/Developer/ShaderCompilerCommon/Private/ShaderCompilerCommon.cpp b/Engine/Source/Developer/ShaderCompilerCommon/Private/ShaderCompilerCommon.cpp
--- a/Engine/Source/Developer/ShaderCompilerCommon/Private/ShaderCompilerCommon.cpp
+++ b/Engine/Source/Developer/ShaderCompilerCommon/Private/ShaderCompilerCommon.cpp
@@ -6,6 +6,7 @@
 #include "Modules/ModuleManager.h"
 #include "HlslccDefinitions.h"
 #include "HAL/FileManager.h"
+#include "String/RemoveFrom.h"
 
 IMPLEMENT_MODULE(FDefaultModuleImpl, ShaderCompilerCommon);
 
@@ -470,21 +471,99 @@ TCHAR* FindNextUniformBufferReference(TCHAR* SearchPtr, const TCHAR* SearchStrin
 	return nullptr;
 }
 
+static const TCHAR* const s_AllSRVTypes[] =
+{
+	TEXT("Texture1D"),
+	TEXT("Texture1DArray"),
+	TEXT("Texture2D"),
+	TEXT("Texture2DArray"),
+	TEXT("Texture2DMS"),
+	TEXT("Texture2DMSArray"),
+	TEXT("Texture3D"),
+	TEXT("TextureCube"),
+	TEXT("TextureCubeArray"),
+
+	TEXT("Buffer"),
+	TEXT("ByteAddressBuffer"),
+	TEXT("StructuredBuffer"),
+	TEXT("ConstantBuffer"),
+	TEXT("RaytracingAccelerationStructure"),
+};
+
+static const TCHAR* const s_AllUAVTypes[] =
+{
+	TEXT("AppendStructuredBuffer"),
+	TEXT("RWBuffer"),
+	TEXT("RWByteAddressBuffer"),
+	TEXT("RWStructuredBuffer"),
+	TEXT("RWTexture1D"),
+	TEXT("RWTexture1DArray"),
+	TEXT("RWTexture2D"),
+	TEXT("RWTexture2DArray"),
+	TEXT("RWTexture3D"),
+	TEXT("RasterizerOrderedTexture2D"),
+};
+
+static const TCHAR* const s_AllSamplerTypes[] =
+{
+	TEXT("SamplerState"),
+	TEXT("SamplerComparisonState"),
+};
+
+EShaderParameterType UE::ShaderCompilerCommon::ParseParameterType(
+	FStringView InType,
+	TArrayView<const TCHAR*> InExtraSRVTypes,
+	TArrayView<const TCHAR*> InExtraUAVTypes)
+{
+	TArrayView<const TCHAR* const> AllSamplerTypes(s_AllSamplerTypes);
+	TArrayView<const TCHAR* const> AllSRVTypes(s_AllSRVTypes);
+	TArrayView<const TCHAR* const> AllUAVTypes(s_AllUAVTypes);
+
+	if (AllSamplerTypes.Contains(InType))
+	{
+		return EShaderParameterType::Sampler;
+	}
+
+	FStringView UntemplatedType = InType;
+	if (int32 Index = InType.Find(TEXT("<")); Index != INDEX_NONE)
+	{
+		const int32 NumChars = InType.Len() - Index;
+		UntemplatedType = InType.LeftChop(NumChars);
+	}
+
+	if (AllSRVTypes.Contains(UntemplatedType) || InExtraSRVTypes.Contains(UntemplatedType))
+	{
+		return EShaderParameterType::SRV;
+	}
+
+	if (AllUAVTypes.Contains(UntemplatedType) || InExtraUAVTypes.Contains(UntemplatedType))
+	{
+		return EShaderParameterType::UAV;
+	}
+
+	return EShaderParameterType::LooseData;
+}
+
+FStringView UE::ShaderCompilerCommon::RemoveConstantBufferPrefix(FStringView InName)
+{
+	return UE::String::RemoveFromStart(InName, FStringView(UE::ShaderCompilerCommon::kUniformBufferConstantBufferPrefix));
+}
+
 FString UE::ShaderCompilerCommon::RemoveConstantBufferPrefix(const FString& InName)
 {
-	FString NewName(InName);
-	NewName.RemoveFromStart(UE::ShaderCompilerCommon::kUniformBufferConstantBufferPrefix);
-	return NewName;
+	return FString(RemoveConstantBufferPrefix(FStringView(InName)));
 }
 
-EShaderParameterType UE::ShaderCompilerCommon::ParseAndRemoveBindlessParameterPrefix(FString& InName)
+EShaderParameterType UE::ShaderCompilerCommon::ParseAndRemoveBindlessParameterPrefix(FStringView& InName)
 {
-	if (InName.RemoveFromStart(UE::ShaderCompilerCommon::kBindlessResourcePrefix))
+	const FStringView OriginalName = InName;
+
+	if (InName = UE::String::RemoveFromStart(InName, FStringView(UE::ShaderCompilerCommon::kBindlessResourcePrefix)); InName != OriginalName)
 	{
 		return EShaderParameterType::BindlessResourceIndex;
 	}
 
-	if (InName.RemoveFromStart(UE::ShaderCompilerCommon::kBindlessSamplerPrefix))
+	if (InName = UE::String::RemoveFromStart(InName, FStringView(UE::ShaderCompilerCommon::kBindlessSamplerPrefix)); InName != OriginalName)
 	{
 		return EShaderParameterType::BindlessSamplerIndex;
 	}
@@ -492,9 +571,19 @@ EShaderParameterType UE::ShaderCompilerCommon::ParseAndRemoveBindlessParameterPr
 	return EShaderParameterType::LooseData;
 }
 
+EShaderParameterType UE::ShaderCompilerCommon::ParseAndRemoveBindlessParameterPrefix(FString& InName)
+{
+	FStringView Name(InName);
+	const EShaderParameterType ParameterType = ParseAndRemoveBindlessParameterPrefix(Name);
+	InName = FString(Name);
+
+	return ParameterType;
+}
+
 bool UE::ShaderCompilerCommon::RemoveBindlessParameterPrefix(FString& InName)
 {
-	return ParseAndRemoveBindlessParameterPrefix(InName) != EShaderParameterType::LooseData;
+	return InName.RemoveFromStart(UE::ShaderCompilerCommon::kBindlessResourcePrefix)
+		|| InName.RemoveFromStart(UE::ShaderCompilerCommon::kBindlessSamplerPrefix);
 }
 
 void HandleReflectedGlobalConstantBufferMember(
diff --git a/Engine/Source/Developer/ShaderCompilerCommon/Private/ShaderParameterParser.cpp b/Engine/Source/Developer/ShaderCompilerCommon/Private/ShaderParameterParser.cpp
--- a/Engine/Source/Developer/ShaderCompilerCommon/Private/ShaderParameterParser.cpp
+++ b/Engine/Source/Developer/ShaderCompilerCommon/Private/ShaderParameterParser.cpp
@@ -2,6 +2,110 @@
 
 #include "ShaderParameterParser.h"
 #include "ShaderCompilerCommon.h"
+#include "Misc/StringBuilder.h"
+
+template<typename TParameterFunction>
+static void IterateShaderParameterMembersInternal(
+	const FShaderParametersMetadata& ParametersMetadata,
+	uint16 ByteOffset,
+	TStringBuilder<1024>& ShaderBindingNameBuilder,
+	TParameterFunction Lambda)
+{
+	for (const FShaderParametersMetadata::FMember& Member : ParametersMetadata.GetMembers())
+	{
+		EUniformBufferBaseType BaseType = Member.GetBaseType();
+		const uint16 MemberOffset = ByteOffset + uint16(Member.GetOffset());
+		const uint32 NumElements = Member.GetNumElements();
+
+		int32 MemberNameLength = FCString::Strlen(Member.GetName());
+
+		if (BaseType == UBMT_INCLUDED_STRUCT)
+		{
+			check(NumElements == 0);
+			const FShaderParametersMetadata& NewParametersMetadata = *Member.GetStructMetadata();
+			IterateShaderParameterMembersInternal(NewParametersMetadata, MemberOffset, ShaderBindingNameBuilder, Lambda);
+		}
+		else if (BaseType == UBMT_NESTED_STRUCT && NumElements == 0)
+		{
+			ShaderBindingNameBuilder.Append(Member.GetName());
+			ShaderBindingNameBuilder.Append(TEXT("_"));
+
+			const FShaderParametersMetadata& NewParametersMetadata = *Member.GetStructMetadata();
+			IterateShaderParameterMembersInternal(NewParametersMetadata, MemberOffset, ShaderBindingNameBuilder, Lambda);
+
+			ShaderBindingNameBuilder.RemoveSuffix(MemberNameLength + 1);
+		}
+		else if (BaseType == UBMT_NESTED_STRUCT && NumElements > 0)
+		{
+			ShaderBindingNameBuilder.Append(Member.GetName());
+			ShaderBindingNameBuilder.Append(TEXT("_"));
+
+			const FShaderParametersMetadata& NewParametersMetadata = *Member.GetStructMetadata();
+			for (uint32 ArrayElementId = 0; ArrayElementId < NumElements; ArrayElementId++)
+			{
+				FString ArrayElementIdString;
+				ArrayElementIdString.AppendInt(ArrayElementId);
+				int32 ArrayElementIdLength = ArrayElementIdString.Len();
+
+				ShaderBindingNameBuilder.Append(ArrayElementIdString);
+				ShaderBindingNameBuilder.Append(TEXT("_"));
+
+				uint16 NewStructOffset = MemberOffset + ArrayElementId * NewParametersMetadata.GetSize();
+				IterateShaderParameterMembersInternal(NewParametersMetadata, NewStructOffset, ShaderBindingNameBuilder, Lambda);
+
+				ShaderBindingNameBuilder.RemoveSuffix(ArrayElementIdLength + 1);
+			}
+
+			ShaderBindingNameBuilder.RemoveSuffix(MemberNameLength + 1);
+		}
+		else
+		{
+			const bool bParametersAreExpanded =
+				NumElements > 0 &&
+				(BaseType == UBMT_TEXTURE ||
+					BaseType == UBMT_SRV ||
+					BaseType == UBMT_UAV ||
+					BaseType == UBMT_SAMPLER ||
+					IsRDGResourceReferenceShaderParameterType(BaseType));
+
+			if (bParametersAreExpanded)
+			{
+				const uint16 ElementSize = SHADER_PARAMETER_POINTER_ALIGNMENT;
+
+				for (uint32 Index = 0; Index < NumElements; Index++)
+				{
+					const FString RealBindingName = FString::Printf(TEXT("%s_%d"), Member.GetName(), Index);
+
+					ShaderBindingNameBuilder.Append(RealBindingName);
+					Lambda(ParametersMetadata, Member, ShaderBindingNameBuilder.ToString(), MemberOffset + Index * ElementSize);
+					ShaderBindingNameBuilder.RemoveSuffix(RealBindingName.Len());
+				}
+			}
+			else
+			{
+				ShaderBindingNameBuilder.Append(Member.GetName());
+				Lambda(ParametersMetadata, Member, ShaderBindingNameBuilder.ToString(), MemberOffset);
+				ShaderBindingNameBuilder.RemoveSuffix(MemberNameLength);
+			}
+		}
+	}
+}
+
+template<typename TParameterFunction>
+static void IterateShaderParameterMembers(const FShaderParametersMetadata& ShaderParametersMetadata, TParameterFunction Lambda)
+{
+	FShaderParametersMetadata::EUseCase UseCase = ShaderParametersMetadata.GetUseCase();
+
+	TStringBuilder<1024> ShaderBindingNameBuilder;
+	if (UseCase == FShaderParametersMetadata::EUseCase::UniformBuffer || UseCase == FShaderParametersMetadata::EUseCase::DataDrivenUniformBuffer)
+	{
+		ShaderBindingNameBuilder.Append(ShaderParametersMetadata.GetShaderVariableName());
+		ShaderBindingNameBuilder.Append(TEXT("_"));
+	}
+
+	IterateShaderParameterMembersInternal(
+		ShaderParametersMetadata, /* ByteOffset = */ 0, ShaderBindingNameBuilder, Lambda);
+}
 
 static void AddNoteToDisplayShaderParameterMemberOnCppSide(
 	const FShaderCompilerInput& CompilerInput,
@@ -28,48 +132,44 @@ static void AddNoteToDisplayShaderParameterMemberOnCppSide(
 	CompilerOutput.Errors.Add(Error);
 }
 
-inline bool MemberWasPotentiallyMoved(const FShaderParametersMetadata::FMember& InMember)
+FShaderParameterParser::FShaderParameterParser() = default;
+FShaderParameterParser::~FShaderParameterParser() = default;
+
+FShaderParameterParser::FShaderParameterParser(
+	TArrayView<const TCHAR*> InExtraSRVTypes,
+	TArrayView<const TCHAR*> InExtraUAVTypes
+)
+	: ExtraSRVTypes(InExtraSRVTypes)
+	, ExtraUAVTypes(InExtraUAVTypes)
 {
-	const EUniformBufferBaseType BaseType = InMember.GetBaseType();
-	return BaseType == UBMT_INT32 || BaseType == UBMT_UINT32 || BaseType == UBMT_FLOAT32
-		|| BaseType == UBMT_TEXTURE || BaseType == UBMT_SRV || BaseType == UBMT_UAV
-		|| BaseType == UBMT_RDG_TEXTURE || BaseType == UBMT_RDG_TEXTURE_SRV || BaseType == UBMT_RDG_TEXTURE_UAV
-		|| BaseType == UBMT_RDG_BUFFER_SRV || BaseType == UBMT_RDG_BUFFER_UAV
-		|| BaseType == UBMT_SAMPLER
-		;
 }
 
-bool FShaderParameterParser::ParseAndMoveShaderParametersToRootConstantBuffer(
+bool FShaderParameterParser::ParseParameters(
 	const FShaderCompilerInput& CompilerInput,
-	FShaderCompilerOutput& CompilerOutput,
-	FString& PreprocessedShaderSource,
-	const TCHAR* ConstantBufferType)
+	FShaderCompilerOutput& CompilerOutput)
 {
-	// The shader doesn't have any parameter binding through shader structure, therefore don't do anything.
-	if (!CompilerInput.RootParametersStructure)
-	{
-		return true;
-	}
-
-	const bool bMoveToRootConstantBuffer = ConstantBufferType != nullptr;
-	OriginalParsedShader = PreprocessedShaderSource;
+	const FStringView ShaderSource(OriginalParsedShader);
 
-	// Reserves the number of parameters up front.
-	ParsedParameters.Reserve(CompilerInput.RootParametersStructure->GetSize() / sizeof(int32));
-
-	CompilerInput.RootParametersStructure->IterateShaderParameterMembers(
-		[&](const FShaderParametersMetadata& ParametersMetadata,
-			const FShaderParametersMetadata::FMember& Member,
-			const TCHAR* ShaderBindingName,
-			uint16 ByteOffset)
+	if (CompilerInput.RootParametersStructure)
 	{
-		FParsedShaderParameter ParsedParameter;
-		ParsedParameter.Member = &Member;
-		ParsedParameter.ConstantBufferOffset = ByteOffset;
-		check(ParsedParameter.IsBindable());
+		// Reserves the number of parameters up front.
+		ParsedParameters.Reserve(CompilerInput.RootParametersStructure->GetSize() / sizeof(int32));
+
+		IterateShaderParameterMembers(
+			*CompilerInput.RootParametersStructure,
+			[&](const FShaderParametersMetadata& ParametersMetadata,
+				const FShaderParametersMetadata::FMember& Member,
+				const TCHAR* ShaderBindingName,
+				uint16 ByteOffset)
+			{
+				FParsedShaderParameter ParsedParameter;
+				ParsedParameter.Member = &Member;
+				ParsedParameter.ConstantBufferOffset = ByteOffset;
+				check(ParsedParameter.IsBindable());
 
-		ParsedParameters.Add(ShaderBindingName, ParsedParameter);
-	});
+				ParsedParameters.Add(ShaderBindingName, ParsedParameter);
+			});
+	}
 
 	bool bSuccess = true;
 
@@ -85,6 +185,7 @@ bool FShaderParameterParser::ParseAndMoveShaderParametersToRootConstantBuffer(
 
 			// Parsing what might be a type of the parameter.
 			ParsingPotentialType,
+			ParsingPotentialTypeTemplateArguments,
 			FinishedPotentialType,
 
 			// Parsing what might be a name of the parameter.
@@ -99,7 +200,7 @@ bool FShaderParameterParser::ParseAndMoveShaderParametersToRootConstantBuffer(
 			FoundParameter,
 		};
 
-		const int32 ShaderSourceLen = PreprocessedShaderSource.Len();
+		const int32 ShaderSourceLen = ShaderSource.Len();
 
 		int32 CurrentPragamLineoffset = -1;
 		int32 CurrentLineoffset = 0;
@@ -145,7 +246,7 @@ bool FShaderParameterParser::ParseAndMoveShaderParametersToRootConstantBuffer(
 
 		for (int32 Cursor = 0; Cursor < ShaderSourceLen; Cursor++)
 		{
-			const TCHAR Char = PreprocessedShaderSource[Cursor];
+			const TCHAR Char = ShaderSource[Cursor];
 
 			auto FoundShaderParameter = [&]()
 			{
@@ -155,18 +256,46 @@ bool FShaderParameterParser::ParseAndMoveShaderParametersToRootConstantBuffer(
 				check(NameStartPos != -1);
 				check(NameEndPos != -1);
 
-				FString Type = PreprocessedShaderSource.Mid(TypeStartPos, TypeEndPos - TypeStartPos + 1);
-				FString Name = PreprocessedShaderSource.Mid(NameStartPos, NameEndPos - NameStartPos + 1);
-				const FString ParsedName = Name;
+				FStringView Type = ShaderSource.Mid(TypeStartPos, TypeEndPos - TypeStartPos + 1);
+				FStringView Name = ShaderSource.Mid(NameStartPos, NameEndPos - NameStartPos + 1);
+
+				FStringView Leftovers = ShaderSource.Mid(NameEndPos + 1, (Cursor - 1) - (NameEndPos + 1) + 1);
 
-				const bool bBindlessParameter = UE::ShaderCompilerCommon::RemoveBindlessParameterPrefix(Name);
+				EShaderParameterType ParsedParameterType = UE::ShaderCompilerCommon::ParseAndRemoveBindlessParameterPrefix(Name);
+				const bool bBindlessIndex = (ParsedParameterType == EShaderParameterType::BindlessResourceIndex || ParsedParameterType == EShaderParameterType::BindlessSamplerIndex);
+
+				EShaderParameterType BindlessConversionType = EShaderParameterType::Num;
+
+				if ((bBindlessResources || bBindlessSamplers) && ParsedParameterType == EShaderParameterType::LooseData)
+				{
+					ParsedParameterType = UE::ShaderCompilerCommon::ParseParameterType(Type, ExtraSRVTypes, ExtraUAVTypes);
+
+					if (bBindlessResources && (ParsedParameterType == EShaderParameterType::SRV || ParsedParameterType == EShaderParameterType::UAV))
+					{
+						BindlessConversionType = EShaderParameterType::BindlessResourceIndex;
+					}
+					else if (bBindlessSamplers && ParsedParameterType == EShaderParameterType::Sampler)
+					{
+						BindlessConversionType = EShaderParameterType::BindlessSamplerIndex;
+					}
+
+					if (BindlessConversionType != EShaderParameterType::Num && Leftovers.Contains(TEXT("register")))
+					{
+						// avoid rewriting hardcoded register assignments
+						BindlessConversionType = EShaderParameterType::Num;
+					}
+				}
 
 				FParsedShaderParameter ParsedParameter;
+
+				EShaderParameterType ConstantBufferParameterType = EShaderParameterType::Num;
+				bool bMoveToRootConstantBuffer = false;
 				bool bUpdateParsedParameters = false;
-				bool bEraseOriginalParameter = false;
-				if (ParsedParameters.Contains(Name))
+
+ 				const FString ParsedParameterKey(Name);
+				if (ParsedParameters.Contains(ParsedParameterKey))
 				{
-					if (ParsedParameters.FindChecked(Name).IsFound())
+					if (ParsedParameters.FindChecked(ParsedParameterKey).IsFound())
 					{
 						// If it has already been found, it means it is duplicated. Do nothing and let the shader compiler throw the error.
 					}
@@ -174,13 +303,23 @@ bool FShaderParameterParser::ParseAndMoveShaderParametersToRootConstantBuffer(
 					{
 						// Update the parsed parameters
 						bUpdateParsedParameters = true;
-						ParsedParameter = ParsedParameters.FindChecked(Name);
+						ParsedParameter = ParsedParameters.FindChecked(ParsedParameterKey);
 
 						// Erase the parameter to move it into the root constant buffer.
-						if (bMoveToRootConstantBuffer && ParsedParameter.IsBindable())
+						if (bNeedToMoveToRootConstantBuffer && ParsedParameter.IsBindable())
 						{
-							EUniformBufferBaseType BaseType = ParsedParameter.Member->GetBaseType();
-							bEraseOriginalParameter = BaseType == UBMT_INT32 || BaseType == UBMT_UINT32 || BaseType == UBMT_FLOAT32 || bBindlessParameter;
+							const EUniformBufferBaseType BaseType = ParsedParameter.Member->GetBaseType();
+							bMoveToRootConstantBuffer =
+								BaseType == UBMT_INT32 ||
+								BaseType == UBMT_UINT32 ||
+								BaseType == UBMT_FLOAT32 ||
+								bBindlessIndex ||
+								(BindlessConversionType != EShaderParameterType::Num);
+
+							if (bMoveToRootConstantBuffer)
+							{
+								ConstantBufferParameterType = BindlessConversionType != EShaderParameterType::Num ? BindlessConversionType : ParsedParameterType;
+							}
 						}
 					}
 				}
@@ -193,28 +332,20 @@ bool FShaderParameterParser::ParseAndMoveShaderParametersToRootConstantBuffer(
 				// Update 
 				if (bUpdateParsedParameters)
 				{
-					ParsedParameter.ParsedName = ParsedName;
 					ParsedParameter.ParsedType = Type;
 					ParsedParameter.ParsedPragmaLineoffset = CurrentPragamLineoffset;
 					ParsedParameter.ParsedLineOffset = CurrentLineoffset;
-					ParsedParameter.bOriginalParameterErased = bEraseOriginalParameter;
+					ParsedParameter.ParsedCharOffsetStart = TypeQualifierStartPos != -1 ? TypeQualifierStartPos : TypeStartPos;
+					ParsedParameter.ParsedCharOffsetEnd = Cursor;
+					ParsedParameter.BindlessConversionType = BindlessConversionType;
+					ParsedParameter.ConstantBufferParameterType = ConstantBufferParameterType;
 
 					if (ArrayStartPos != -1 && ArrayEndPos != -1)
 					{
-						ParsedParameter.ParsedArraySize = PreprocessedShaderSource.Mid(ArrayStartPos + 1, ArrayEndPos - ArrayStartPos - 1);
+						ParsedParameter.ParsedArraySize = ShaderSource.Mid(ArrayStartPos + 1, ArrayEndPos - ArrayStartPos - 1);
 					}
 
-					ParsedParameters.Add(Name, ParsedParameter);
-				}
-
-				// Erases this shader parameter conserving the same line numbers.
-				if (bEraseOriginalParameter)
-				{
-					for (int32 j = (TypeQualifierStartPos != -1 ? TypeQualifierStartPos : TypeStartPos); j <= Cursor; j++)
-					{
-						if (PreprocessedShaderSource[j] != '\r' && PreprocessedShaderSource[j] != '\n')
-							PreprocessedShaderSource[j] = ' ';
-					}
+					ParsedParameters.Add(ParsedParameterKey, ParsedParameter);
 				}
 
 				ResetState();
@@ -224,7 +355,7 @@ bool FShaderParameterParser::ParseAndMoveShaderParametersToRootConstantBuffer(
 			const bool bIsLetter = (Char >= 'a' && Char <= 'z') || (Char >= 'A' && Char <= 'Z');
 			const bool bIsNumber = Char >= '0' && Char <= '9';
 
-			const TCHAR* UpComing = (*PreprocessedShaderSource) + Cursor;
+			const TCHAR* UpComing = ShaderSource.GetData() + Cursor;
 			const int32 RemainingSize = ShaderSourceLen - Cursor;
 
 			CurrentLineoffset += Char == '\n';
@@ -272,14 +403,15 @@ bool FShaderParameterParser::ParseAndMoveShaderParametersToRootConstantBuffer(
 			{
 				if (bIsLetter)
 				{
-					static const TCHAR* KeywordTable[] = {
+					static const TCHAR* KeywordTable[] =
+					{
 						TEXT("enum"),
 						TEXT("class"),
 						TEXT("const"),
 						TEXT("struct"),
 						TEXT("static"),
 					};
-					static int32 KeywordTableSize[] = {4, 5, 5, 6, 6};
+					static int32 KeywordTableSize[] = { 4, 5, 5, 6, 6 };
 
 					int32 RecognisedKeywordId = -1;
 					for (int32 KeywordId = 0; KeywordId < UE_ARRAY_COUNT(KeywordTable); KeywordId++)
@@ -349,10 +481,15 @@ bool FShaderParameterParser::ParseAndMoveShaderParametersToRootConstantBuffer(
 				// Found character legal for a type...
 				if (bIsLetter ||
 					bIsNumber ||
-					Char == '<' || Char == '>' || Char == '_')
+					Char == '_')
 				{
 					// Keep browsing what might be type of the parameter.
 				}
+				else if (Char == '<')
+				{
+					// Found what looks like the begining of template argument that is legal on resource types for Instance Texture2D< float >
+					State = EState::ParsingPotentialTypeTemplateArguments;
+				}
 				else if (bIsWhiteSpace)
 				{
 					// Might have found a type.
@@ -365,6 +502,31 @@ bool FShaderParameterParser::ParseAndMoveShaderParametersToRootConstantBuffer(
 					State = EState::GoToNextSemicolonAndReset;
 				}
 			}
+			else if (State == EState::ParsingPotentialTypeTemplateArguments)
+			{
+				// Found character legal for a template argument...
+				if (bIsLetter ||
+					bIsNumber ||
+					Char == '_')
+				{
+					// Keep browsing what might be type of the parameter.
+				}
+				else if (bIsWhiteSpace || Char == ',')
+				{
+					// Spaces and comas are legal within agrument of the template arguments.
+				}
+				else if (Char == '>')
+				{
+					// Might have found a type with template argument.
+					State = EState::FinishedPotentialType;
+					TypeEndPos = Cursor;
+				}
+				else
+				{
+					// Found unexpected character in the type.
+					State = EState::GoToNextSemicolonAndReset;
+				}
+			}
 			else if (State == EState::FinishedPotentialType)
 			{
 				if (bIsLetter)
@@ -373,6 +535,11 @@ bool FShaderParameterParser::ParseAndMoveShaderParametersToRootConstantBuffer(
 					State = EState::ParsingPotentialName;
 					NameStartPos = Cursor;
 				}
+				else if (Char == '<')
+				{
+					// Might have found beginning of a template argument for the type, that was separate by a whitespace from type. For instance Texture2D <float>
+					State = EState::ParsingPotentialTypeTemplateArguments;
+				}
 				else if (bIsWhiteSpace)
 				{
 					// Keep parsing void.
@@ -430,9 +597,7 @@ bool FShaderParameterParser::ParseAndMoveShaderParametersToRootConstantBuffer(
 					State = EState::GoToNextSemicolonAndReset;
 				}
 			}
-			else if (
-				State == EState::FinishedPotentialName ||
-				State == EState::FinishedArraySize)
+			else if (State == EState::FinishedPotentialName || State == EState::FinishedArraySize)
 			{
 				if (Char == ';')
 				{
@@ -508,58 +673,213 @@ bool FShaderParameterParser::ParseAndMoveShaderParametersToRootConstantBuffer(
 			{
 				unimplemented();
 			}
-		} // for (int32 Cursor = 0; Cursor < PreprocessedShaderSource.Len(); Cursor++)
+		} // for (int32 Cursor = 0; Cursor < ShaderSourceLen; Cursor++)
 	}
 
+	return bSuccess;
+}
+
+void FShaderParameterParser::RemoveMovingParametersFromSource(FString& PreprocessedShaderSource)
+{
+	for (TPair<FString, FParsedShaderParameter>& Itr : ParsedParameters)
+	{
+		const FParsedShaderParameter& ParsedParameter = Itr.Value;
+
+		// If this parameter is going to be in the root constant buffer
+		if (ParsedParameter.ConstantBufferParameterType != EShaderParameterType::Num &&
+			// but it's not being converted to bindless
+			ParsedParameter.BindlessConversionType == EShaderParameterType::Num)
+		{
+			// then erase this shader parameter conserving the same line numbers.
+			for (int32 j = ParsedParameter.ParsedCharOffsetStart; j <= ParsedParameter.ParsedCharOffsetEnd; j++)
+			{
+				if (PreprocessedShaderSource[j] != '\r' && PreprocessedShaderSource[j] != '\n')
+				{
+					PreprocessedShaderSource[j] = ' ';
+				}
+			}
+		}
+	}
+}
+
+void FShaderParameterParser::ApplyBindlessModifications(FString& PreprocessedShaderSource)
+{
+	if (bBindlessResources || bBindlessSamplers)
+	{
+		// Array of modifications to do on PreprocessedShaderSource
+		struct FShaderCodeModifications
+		{
+			int32 CharOffsetStart;
+			int32 CharOffsetEnd;
+			FString Replace;
+		};
+		TArray<FShaderCodeModifications> Modifications;
+		Modifications.Reserve(ParsedParameters.Num());
+
+		for (TPair<FString, FParsedShaderParameter>& Itr : ParsedParameters)
+		{
+			const FStringView ShaderBindingName = Itr.Key;
+			FParsedShaderParameter& ParsedParameter = Itr.Value;
+
+			if (!ParsedParameter.IsFound())
+			{
+				continue;
+			}
+
+			if (ParsedParameter.BindlessConversionType != EShaderParameterType::Num)
+			{
+				const bool bIsSampler = (ParsedParameter.BindlessConversionType == EShaderParameterType::BindlessSamplerIndex);
+
+				const TCHAR* IndexDefine = bIsSampler ? TEXT("AUTO_BINDLESS_SAMPLER_INDEX") : TEXT("AUTO_BINDLESS_RESOURCE_INDEX");
+				const TCHAR* VariableDefine = bIsSampler ? TEXT("AUTO_BINDLESS_SAMPLER_VARIABLE") : TEXT("AUTO_BINDLESS_RESOURCE_VARIABLE");
+
+				const FStringView Name = ShaderBindingName;
+				const FStringView Type = ParsedParameter.ParsedType;
+
+				const FString IndexText = FString::Printf(TEXT("%s(%.*s);"), IndexDefine, Name.Len(), Name.GetData());
+				const FString VariableText = FString::Printf(TEXT("%s(%.*s, %.*s);"), VariableDefine, Type.Len(), Type.GetData(), Name.Len(), Name.GetData());
+
+				FShaderCodeModifications Modif;
+				Modif.CharOffsetStart = ParsedParameter.ParsedCharOffsetStart;
+				Modif.CharOffsetEnd = ParsedParameter.ParsedCharOffsetEnd + 1;
+
+				// If we weren't going to be added to a root constant buffer, that means we need to declare our index before we declare our getter.
+				if (ParsedParameter.ConstantBufferParameterType == EShaderParameterType::Num)
+				{
+					Modif.Replace += IndexText;
+				}
+				Modif.Replace += VariableText;
+
+				Modifications.Add(Modif);
+			}
+		}
+
+		// Apply all modifications
+		if (Modifications.Num() > 0)
+		{
+			// Sort all the modifications in order
+			Modifications.Sort(
+				[](const FShaderCodeModifications& ModifA, const FShaderCodeModifications& ModifB)
+				{
+					return  ModifA.CharOffsetStart < ModifB.CharOffsetStart;
+				}
+			);
+
+			// Find out the size of the shader code after all modifications
+			int32 NewShaderCodeSize = PreprocessedShaderSource.Len();
+			for (const FShaderCodeModifications& Modif : Modifications)
+			{
+				// Count the number of line return \n in CharOffsetStart -> CharOffsetEnd to ensure the line number remain unchanged.
+				check(!Modif.Replace.Contains(TEXT("\n")));
+				//int32 ReplacedCarriagedReturn = 0;
+				for (int32 CharPos = Modif.CharOffsetStart; CharPos < Modif.CharOffsetEnd; CharPos++)
+				{
+					ensure(PreprocessedShaderSource[CharPos] != '\n');
+				}
+
+				NewShaderCodeSize += Modif.Replace.Len() - (Modif.CharOffsetEnd - Modif.CharOffsetStart); // + ReplacedCarriagedReturn;
+			}
+
+			// Splice all the code and modifications together
+			FString NewShaderCode;
+			NewShaderCode.Reserve(NewShaderCodeSize);
+
+			int32 CurrentCodePos = 0;
+			for (const FShaderCodeModifications& Modif : Modifications)
+			{
+				check(CurrentCodePos <= Modif.CharOffsetStart);
+				NewShaderCode += PreprocessedShaderSource.Mid(CurrentCodePos, Modif.CharOffsetStart - CurrentCodePos);
+				NewShaderCode += Modif.Replace;
+				CurrentCodePos = Modif.CharOffsetEnd;
+			}
+			check(CurrentCodePos <= PreprocessedShaderSource.Len());
+			NewShaderCode += PreprocessedShaderSource.Mid(CurrentCodePos, PreprocessedShaderSource.Len() - CurrentCodePos);
+
+			// Commit all modifications to caller
+			PreprocessedShaderSource = NewShaderCode;
+		}
+	}
+}
+
+static const TCHAR* GetConstantSwizzle(uint16 ByteOffset)
+{
+	switch (ByteOffset % 16)
+	{
+	default: unimplemented();
+	case 0:  return TEXT("");
+	case 4:  return TEXT(".y");
+	case 8:  return TEXT(".z");
+	case 12: return TEXT(".w");
+	}
+}
+
+bool FShaderParameterParser::MoveShaderParametersToRootConstantBuffer(
+	const FShaderCompilerInput& CompilerInput,
+	FShaderCompilerOutput& CompilerOutput,
+	FString& PreprocessedShaderSource,
+	const TCHAR* ConstantBufferType)
+{
+	bool bSuccess = true;
+
 	// Generate the root cbuffer content.
-	if (bMoveToRootConstantBuffer)
+	if (CompilerInput.RootParametersStructure && bNeedToMoveToRootConstantBuffer)
 	{
 		FString RootCBufferContent;
 
-		CompilerInput.RootParametersStructure->IterateShaderParameterMembers(
+		IterateShaderParameterMembers(
+			*CompilerInput.RootParametersStructure,
 			[&](const FShaderParametersMetadata& ParametersMetadata,
 				const FShaderParametersMetadata::FMember& Member,
 				const TCHAR* ShaderBindingName,
 				uint16 ByteOffset)
 		{
-			if (MemberWasPotentiallyMoved(Member))
+			FParsedShaderParameter* ParsedParameter = ParsedParameters.Find(ShaderBindingName);
+			if (ParsedParameter && ParsedParameter->IsFound() && ParsedParameter->ConstantBufferParameterType != EShaderParameterType::Num)
 			{
-				FParsedShaderParameter* ParsedParameter = ParsedParameters.Find(ShaderBindingName);
-				if (ParsedParameter && ParsedParameter->IsFound() && ParsedParameter->bOriginalParameterErased)
-				{
-					uint32 ConstantRegister = ByteOffset / 16;
-					const TCHAR* ConstantSwizzle = [ByteOffset]()
-						{
-							switch (ByteOffset % 16)
-							{
-							default: unimplemented();
-							case 0:  return TEXT("");
-							case 4:  return TEXT(".y");
-							case 8:  return TEXT(".z");
-							case 12: return TEXT(".w");
-							}
-						}();
+				const uint32 ConstantRegister = ByteOffset / 16;
+				const TCHAR* ConstantSwizzle = GetConstantSwizzle(ByteOffset);
 
+#define SVARG(N) N.Len(), N.GetData()
+				if (ParsedParameter->ConstantBufferParameterType == EShaderParameterType::BindlessResourceIndex)
+				{
+					RootCBufferContent.Append(FString::Printf(
+						TEXT("uint BindlessResource_%s : packoffset(c%d%s);\r\n"),
+						ShaderBindingName,
+						ConstantRegister,
+						ConstantSwizzle));
+				}
+				else if (ParsedParameter->ConstantBufferParameterType == EShaderParameterType::BindlessSamplerIndex)
+				{
+					RootCBufferContent.Append(FString::Printf(
+						TEXT("uint BindlessSampler_%s : packoffset(c%d%s);\r\n"),
+						ShaderBindingName,
+						ConstantRegister,
+						ConstantSwizzle));
+				}
+				else if (ParsedParameter->ConstantBufferParameterType == EShaderParameterType::LooseData)
+				{
 					if (!ParsedParameter->ParsedArraySize.IsEmpty())
 					{
 						RootCBufferContent.Append(FString::Printf(
-							TEXT("%s %s[%s] : packoffset(c%d%s);\r\n"),
-							*ParsedParameter->ParsedType,
-							*ParsedParameter->ParsedName,
-							*ParsedParameter->ParsedArraySize,
+							TEXT("%.*s %s[%.*s] : packoffset(c%d%s);\r\n"),
+							SVARG(ParsedParameter->ParsedType),
+							ShaderBindingName,
+							SVARG(ParsedParameter->ParsedArraySize),
 							ConstantRegister,
 							ConstantSwizzle));
 					}
 					else
 					{
 						RootCBufferContent.Append(FString::Printf(
-							TEXT("%s %s : packoffset(c%d%s);\r\n"),
-							*ParsedParameter->ParsedType,
-							*ParsedParameter->ParsedName,
+							TEXT("%.*s %s : packoffset(c%d%s);\r\n"),
+							SVARG(ParsedParameter->ParsedType),
+							ShaderBindingName,
 							ConstantRegister,
 							ConstantSwizzle));
 					}
 				}
+#undef SVARG
+
 			}
 		});
 
@@ -573,17 +893,52 @@ bool FShaderParameterParser::ParseAndMoveShaderParametersToRootConstantBuffer(
 			*RootCBufferContent);
 
 		FString NewShaderCode = (
-			MakeInjectedShaderCodeBlock(TEXT("ParseAndMoveShaderParametersToRootConstantBuffer"), CBufferCodeBlock) +
+			MakeInjectedShaderCodeBlock(TEXT("MoveShaderParametersToRootConstantBuffer"), CBufferCodeBlock) +
 			PreprocessedShaderSource);
 
 		PreprocessedShaderSource = MoveTemp(NewShaderCode);
 
 		bMovedLoosedParametersToRootConstantBuffer = true;
-	} // if (bMoveToRootConstantBuffer)
+	} // if (CompilerInput.RootParametersStructure && bNeedToMoveToRootConstantBuffer)
 
 	return bSuccess;
 }
 
+bool FShaderParameterParser::ParseAndModify(
+	const FShaderCompilerInput& CompilerInput,
+	FShaderCompilerOutput& CompilerOutput,
+	FString& PreprocessedShaderSource,
+	const TCHAR* ConstantBufferType)
+{
+	bBindlessResources = CompilerInput.Environment.CompilerFlags.Contains(CFLAG_BindlessResources);
+	bBindlessSamplers = CompilerInput.Environment.CompilerFlags.Contains(CFLAG_BindlessSamplers);
+	bHasRootParameters = (CompilerInput.RootParametersStructure != nullptr);
+
+	// The shader doesn't have any parameter binding through shader structure, therefore don't do anything.
+	if (!(bBindlessResources || bBindlessSamplers || bHasRootParameters))
+	{
+		return true;
+	}
+
+	bNeedToMoveToRootConstantBuffer = ConstantBufferType != nullptr;
+	OriginalParsedShader = PreprocessedShaderSource;
+
+	if (!ParseParameters(CompilerInput, CompilerOutput))
+	{
+		return false;
+	}
+
+	RemoveMovingParametersFromSource(PreprocessedShaderSource);
+	ApplyBindlessModifications(PreprocessedShaderSource);
+
+	if (bNeedToMoveToRootConstantBuffer)
+	{
+		return MoveShaderParametersToRootConstantBuffer(CompilerInput, CompilerOutput, PreprocessedShaderSource, ConstantBufferType);
+	}
+	
+	return true;
+}
+
 void FShaderParameterParser::ValidateShaderParameterType(
 	const FShaderCompilerInput& CompilerInput,
 	const FString& ShaderBindingName,
@@ -621,24 +976,32 @@ void FShaderParameterParser::ValidateShaderParameterType(
 
 		if (!bIsTypeCorrect)
 		{
+			auto CheckTypeCorrect = [&ParsedParameter, &ExpectedShaderType](int32 ParsedOffset, int32 ExpectedOffset) -> bool
+			{
+				const FStringView Parsed = ParsedParameter.ParsedType.RightChop(ParsedOffset);
+				const FStringView Expected = FStringView(ExpectedShaderType).RightChop(ExpectedOffset);
+
+				return Parsed.Compare(Expected, ESearchCase::CaseSensitive) == 0;
+			};
+
 			// Accept half-precision floats when single-precision was requested
 			if (ParsedParameter.ParsedType.StartsWith(TEXT("half")) && ParsedParameter.Member->GetBaseType() == UBMT_FLOAT32)
 			{
-				bIsTypeCorrect = (FCString::Strcmp(*ParsedParameter.ParsedType + 4, *ExpectedShaderType + 5) == 0);
+				bIsTypeCorrect = CheckTypeCorrect(4, 5);
 			}
 			// Accept single-precision floats when half-precision was expected
 			else if (ParsedParameter.ParsedType.StartsWith(TEXT("float")) && ExpectedShaderType.StartsWith(TEXT("half")))
 			{
-				bIsTypeCorrect = (FCString::Strcmp(*ParsedParameter.ParsedType + 5, *ExpectedShaderType + 4) == 0);
+				bIsTypeCorrect = CheckTypeCorrect(5, 4);
 			}
 			// support for min16float
 			else if (ParsedParameter.ParsedType.StartsWith(TEXT("min16float")) && ExpectedShaderType.StartsWith(TEXT("float")))
 			{
-				bIsTypeCorrect = (FCString::Strcmp(*ParsedParameter.ParsedType + 10, *ExpectedShaderType + 5) == 0);
+				bIsTypeCorrect = CheckTypeCorrect(10, 5);
 			}
 			else if (ParsedParameter.ParsedType.StartsWith(TEXT("min16float")) && ExpectedShaderType.StartsWith(TEXT("half")))
 			{
-				bIsTypeCorrect = (FCString::Strcmp(*ParsedParameter.ParsedType + 10, *ExpectedShaderType + 4) == 0);
+				bIsTypeCorrect = CheckTypeCorrect(10, 4);
 			}
 		}
 
@@ -667,8 +1030,8 @@ void FShaderParameterParser::ValidateShaderParameterType(
 
 			FShaderCompilerError Error;
 			Error.StrippedErrorMessage = FString::Printf(
-				TEXT("Error: Type %s of shader parameter %s in shader mismatch the shader parameter structure: %s expects a %s"),
-				*ParsedParameter.ParsedType,
+				TEXT("Error: Type %.*s of shader parameter %s in shader mismatch the shader parameter structure: %s expects a %s"),
+				ParsedParameter.ParsedType.Len(), ParsedParameter.ParsedType.GetData(),
 				*ShaderBindingName,
 				*CppCodeName,
 				*ExpectedShaderType);
@@ -720,7 +1083,8 @@ void FShaderParameterParser::ValidateShaderParameterTypes(
 
 	const TMap<FString, FParameterAllocation>& ParametersFoundByCompiler = CompilerOutput.ParameterMap.GetParameterMap();
 
-	CompilerInput.RootParametersStructure->IterateShaderParameterMembers(
+	IterateShaderParameterMembers(
+		*CompilerInput.RootParametersStructure,
 		[&](const FShaderParametersMetadata& ParametersMetadata,
 			const FShaderParametersMetadata::FMember& Member,
 			const TCHAR* ShaderBindingName,
diff --git a/Engine/Source/Developer/ShaderCompilerCommon/Public/ShaderCompilerCommon.h b/Engine/Source/Developer/ShaderCompilerCommon/Public/ShaderCompilerCommon.h
--- a/Engine/Source/Developer/ShaderCompilerCommon/Public/ShaderCompilerCommon.h
+++ b/Engine/Source/Developer/ShaderCompilerCommon/Public/ShaderCompilerCommon.h
@@ -56,8 +56,14 @@ extern SHADERCOMPILERCOMMON_API int16 GetNumUniformBuffersUsed(const FShaderComp
 
 namespace UE::ShaderCompilerCommon
 {
+	extern SHADERCOMPILERCOMMON_API EShaderParameterType ParseParameterType(FStringView InType, TArrayView<const TCHAR*> InExtraSRVTypes, TArrayView<const TCHAR*> InExtraUAVTypes);
+
+	extern SHADERCOMPILERCOMMON_API FStringView          RemoveConstantBufferPrefix(FStringView InName);
 	extern SHADERCOMPILERCOMMON_API FString              RemoveConstantBufferPrefix(const FString& InName);
+
+	extern SHADERCOMPILERCOMMON_API EShaderParameterType ParseAndRemoveBindlessParameterPrefix(FStringView& InName);
 	extern SHADERCOMPILERCOMMON_API EShaderParameterType ParseAndRemoveBindlessParameterPrefix(FString& InName);
+
 	extern SHADERCOMPILERCOMMON_API bool                 RemoveBindlessParameterPrefix(FString& InName);
 }
 
diff --git a/Engine/Source/Developer/ShaderCompilerCommon/Public/ShaderParameterParser.h b/Engine/Source/Developer/ShaderCompilerCommon/Public/ShaderParameterParser.h
--- a/Engine/Source/Developer/ShaderCompilerCommon/Public/ShaderParameterParser.h
+++ b/Engine/Source/Developer/ShaderCompilerCommon/Public/ShaderParameterParser.h
@@ -8,6 +8,8 @@
 struct FShaderCompilerInput;
 struct FShaderCompilerOutput;
 
+enum class EShaderParameterType : uint8;
+
 /** Validates and moves all the shader loose data parameter defined in the root scope of the shader into the root uniform buffer. */
 class SHADERCOMPILERCOMMON_API FShaderParameterParser
 {
@@ -19,16 +21,12 @@ public:
 		const FShaderParametersMetadata::FMember* Member = nullptr;
 
 		/** Information found about the member when parsing the preprocessed code. */
-		FString ParsedName;
-		FString ParsedType;
-		FString ParsedArraySize;
+		FStringView ParsedType; /** View into FShaderParameterParser::OriginalParsedShader */
+		FStringView ParsedArraySize; /** View into FShaderParameterParser::OriginalParsedShader */
 
 		/** Offset the member should be in the constant buffer. */
 		int32 ConstantBufferOffset = 0;
 
-		/** Records if the parameter was fully moved. */
-		bool bOriginalParameterErased{};
-
 		/* Returns whether the shader parameter has been found when parsing. */
 		bool IsFound() const
 		{
@@ -45,15 +43,38 @@ public:
 		int32 ParsedPragmaLineoffset = 0;
 		int32 ParsedLineOffset = 0;
 
+		/** Character position of the start and end of the parameter decelaration in FParsedShaderParameter::OriginalParsedShader */
+		int32 ParsedCharOffsetStart = 0;
+		int32 ParsedCharOffsetEnd = 0;
+
+		EShaderParameterType BindlessConversionType{};
+		EShaderParameterType ConstantBufferParameterType{};
+
 		friend class FShaderParameterParser;
 	};
 
+	FShaderParameterParser();
+	FShaderParameterParser(TArrayView<const TCHAR*> InExtraSRVTypes, TArrayView<const TCHAR*> InExtraUAVTypes);
+	~FShaderParameterParser();
+
+	/** Parses the preprocessed shader code and applies the necessary modifications to it. */
+	bool ParseAndModify(
+		const FShaderCompilerInput& CompilerInput,
+		FShaderCompilerOutput& CompilerOutput,
+		FString& PreprocessedShaderSource,
+		const TCHAR* ConstantBufferType
+	);
+
 	/** Parses the preprocessed shader code and move the parameters into root constant buffer */
+	UE_DEPRECATED(5.1, "ParseAndModify should be called instead.")
 	bool ParseAndMoveShaderParametersToRootConstantBuffer(
 		const FShaderCompilerInput& CompilerInput,
 		FShaderCompilerOutput& CompilerOutput,
 		FString& PreprocessedShaderSource,
-		const TCHAR* ConstantBufferType);
+		const TCHAR* ConstantBufferType)
+	{
+		return ParseAndModify(CompilerInput, CompilerOutput, PreprocessedShaderSource, ConstantBufferType);
+	}
 
 	/** Gets parsing information from a parameter binding name. */
 	const FParsedShaderParameter& FindParameterInfos(const FString& ParameterName) const
@@ -100,11 +121,45 @@ public:
 	}
 
 private:
+	/** Parses the preprocessed shader code */
+	bool ParseParameters(
+		const FShaderCompilerInput& CompilerInput,
+		FShaderCompilerOutput& CompilerOutput
+	);
+
+	void RemoveMovingParametersFromSource(
+		FString& PreprocessedShaderSource
+	);
+
+	/** Converts parsed parameters into their bindless forms. */
+	void ApplyBindlessModifications(
+		FString& PreprocessedShaderSource
+	);
+
+	/** Moves parsed parameters into the root constant buffer. */
+	bool MoveShaderParametersToRootConstantBuffer(
+		const FShaderCompilerInput& CompilerInput,
+		FShaderCompilerOutput& CompilerOutput,
+		FString& PreprocessedShaderSource,
+		const TCHAR* ConstantBufferType
+	);
+
 	void ExtractFileAndLine(int32 PragamLineoffset, int32 LineOffset, FString& OutFile, FString& OutLine) const;
 
+	TArray<const TCHAR*> ExtraSRVTypes;
+	TArray<const TCHAR*> ExtraUAVTypes;
+
 	FString OriginalParsedShader;
 
 	TMap<FString, FParsedShaderParameter> ParsedParameters;
 
+	bool bHasRootParameters = false;
+	bool bBindlessResources = false;
+	bool bBindlessSamplers = false;
+
+	/** Indicates that parameters should be moved to the root cosntant buffer. */
+	bool bNeedToMoveToRootConstantBuffer = false;
+
+	/** Indicates that parameters were actually moved to the root constant buffer. */
 	bool bMovedLoosedParametersToRootConstantBuffer = false;
 };
diff --git a/Engine/Source/Developer/ShaderFormatOpenGL/Private/OpenGLShaderCompiler.cpp b/Engine/Source/Developer/ShaderFormatOpenGL/Private/OpenGLShaderCompiler.cpp
--- a/Engine/Source/Developer/ShaderFormatOpenGL/Private/OpenGLShaderCompiler.cpp
+++ b/Engine/Source/Developer/ShaderFormatOpenGL/Private/OpenGLShaderCompiler.cpp
@@ -3325,8 +3325,7 @@ void FOpenGLFrontend::CompileShader(const FShaderCompilerInput& Input, FShaderCo
 	}
 
 	FShaderParameterParser ShaderParameterParser;
-	if (!ShaderParameterParser.ParseAndMoveShaderParametersToRootConstantBuffer(
-		Input, Output, PreprocessedShader, /* ConstantBufferType = */ nullptr))
+	if (!ShaderParameterParser.ParseAndModify(Input, Output, PreprocessedShader, /* ConstantBufferType = */ nullptr))
 	{
 		// The FShaderParameterParser will add any relevant errors.
 		return;
diff --git a/Engine/Source/Developer/VulkanShaderFormat/Private/VulkanShaderCompiler.cpp b/Engine/Source/Developer/VulkanShaderFormat/Private/VulkanShaderCompiler.cpp
--- a/Engine/Source/Developer/VulkanShaderFormat/Private/VulkanShaderCompiler.cpp
+++ b/Engine/Source/Developer/VulkanShaderFormat/Private/VulkanShaderCompiler.cpp
@@ -2237,8 +2237,7 @@ void DoCompileVulkanShader(const FShaderCompilerInput& Input, FShaderCompilerOut
 	}
 
 	FShaderParameterParser ShaderParameterParser;
-	if (!ShaderParameterParser.ParseAndMoveShaderParametersToRootConstantBuffer(
-		Input, Output, PreprocessedShaderSource, /* ConstantBufferType = */ nullptr))
+	if (!ShaderParameterParser.ParseAndModify(Input, Output, PreprocessedShaderSource, /* ConstantBufferType = */ nullptr))
 	{
 		// The FShaderParameterParser will add any relevant errors.
 		return;
diff --git a/Engine/Source/Developer/Windows/ShaderFormatD3D/Private/D3DShaderCompiler.cpp b/Engine/Source/Developer/Windows/ShaderFormatD3D/Private/D3DShaderCompiler.cpp
--- a/Engine/Source/Developer/Windows/ShaderFormatD3D/Private/D3DShaderCompiler.cpp
+++ b/Engine/Source/Developer/Windows/ShaderFormatD3D/Private/D3DShaderCompiler.cpp
@@ -1297,7 +1297,7 @@ void CompileD3DShader(const FShaderCompilerInput& Input, FShaderCompilerOutput&
 	}
 
 	FShaderParameterParser ShaderParameterParser;
-	if (!ShaderParameterParser.ParseAndMoveShaderParametersToRootConstantBuffer(
+	if (!ShaderParameterParser.ParseAndModify(
 		Input, Output, PreprocessedShaderSource,
 		(Input.IsRayTracingShader() || ShouldUseStableConstantBuffer(Input)) ? TEXT("cbuffer") : nullptr))
 	{
diff --git a/Engine/Source/Developer/Windows/ShaderFormatD3D/Private/D3DShaderCompiler.inl b/Engine/Source/Developer/Windows/ShaderFormatD3D/Private/D3DShaderCompiler.inl
--- a/Engine/Source/Developer/Windows/ShaderFormatD3D/Private/D3DShaderCompiler.inl
+++ b/Engine/Source/Developer/Windows/ShaderFormatD3D/Private/D3DShaderCompiler.inl
@@ -317,7 +317,7 @@ template <typename ID3D1xShaderReflection, typename D3D1x_SHADER_DESC, typename
 				if (ShouldUseStableConstantBuffer(Input))
 				{
 					// Each member found in the global constant buffer means it was not in RootParametersStructure or
-					// it would have been moved by ShaderParameterParser.ParseAndMoveShaderParametersToRootConstantBuffer().
+					// it would have been moved by ShaderParameterParser.ParseAndModify().
 					for (uint32 ConstantIndex = 0; ConstantIndex < CBDesc.Variables; ConstantIndex++)
 					{
 						ID3D1xShaderReflectionVariable* Variable = ConstantBuffer->GetVariableByIndex(ConstantIndex);
diff --git a/Engine/Source/Runtime/Engine/Private/ShaderCompiler/ShaderCompiler.cpp b/Engine/Source/Runtime/Engine/Private/ShaderCompiler/ShaderCompiler.cpp
--- a/Engine/Source/Runtime/Engine/Private/ShaderCompiler/ShaderCompiler.cpp
+++ b/Engine/Source/Runtime/Engine/Private/ShaderCompiler/ShaderCompiler.cpp
@@ -6068,11 +6068,17 @@ void GlobalBeginCompileShader(
 		const ERHIBindlessConfiguration ResourcesConfig = RHIGetBindlessResourcesConfiguration(EShaderPlatform(Target.Platform));
 		const ERHIBindlessConfiguration SamplersConfig = RHIGetBindlessSamplersConfiguration(EShaderPlatform(Target.Platform));
 
-		Input.Environment.SetDefine(TEXT("ENABLE_BINDLESS_RESOURCES"),
-			ResourcesConfig == ERHIBindlessConfiguration::AllShaders || (ResourcesConfig == ERHIBindlessConfiguration::RayTracingShaders && bIsRaytracingShader));
+		if (ResourcesConfig == ERHIBindlessConfiguration::AllShaders || (ResourcesConfig == ERHIBindlessConfiguration::RayTracingShaders && bIsRaytracingShader))
+		{
+			Input.Environment.CompilerFlags.Add(CFLAG_BindlessResources);
+			Input.Environment.SetDefine(TEXT("ENABLE_BINDLESS_RESOURCES"), true);
+		}
 
-		Input.Environment.SetDefine(TEXT("ENABLE_BINDLESS_SAMPLERS"),
-			SamplersConfig == ERHIBindlessConfiguration::AllShaders || (SamplersConfig == ERHIBindlessConfiguration::RayTracingShaders && bIsRaytracingShader));
+		if (SamplersConfig == ERHIBindlessConfiguration::AllShaders || (SamplersConfig == ERHIBindlessConfiguration::RayTracingShaders && bIsRaytracingShader))
+		{
+			Input.Environment.CompilerFlags.Add(CFLAG_BindlessSamplers);
+			Input.Environment.SetDefine(TEXT("ENABLE_BINDLESS_SAMPLERS"), true);
+		}
 	}
 
 	{
diff --git a/Engine/Source/Runtime/RenderCore/Public/ShaderCompilerCore.h b/Engine/Source/Runtime/RenderCore/Public/ShaderCompilerCore.h
--- a/Engine/Source/Runtime/RenderCore/Public/ShaderCompilerCore.h
+++ b/Engine/Source/Runtime/RenderCore/Public/ShaderCompilerCore.h
@@ -117,9 +117,12 @@ enum ECompilerFlags
 	CFLAG_PrecompileWithDXC,
 	// Enable HLSL 2021 version. Enables templates, operator overloaing, C++ style function overloading. Contains breaking change with short-circuiting evaluation.
 	CFLAG_HLSL2021,
-
 	// Allow warnings to be treated as errors
 	CFLAG_WarningsAsErrors,
+	// Enabled if bindless resources are enabled for the platform
+	CFLAG_BindlessResources,
+	// Enabled if bindless samplers are enabled for the platform
+	CFLAG_BindlessSamplers,
 
 	CFLAG_Max,
 };
diff --git a/Engine/Source/Runtime/RenderCore/Public/ShaderParameterMetadata.h b/Engine/Source/Runtime/RenderCore/Public/ShaderParameterMetadata.h
--- a/Engine/Source/Runtime/RenderCore/Public/ShaderParameterMetadata.h
+++ b/Engine/Source/Runtime/RenderCore/Public/ShaderParameterMetadata.h
@@ -10,7 +10,6 @@
 #include "Containers/List.h"
 #include "Containers/StaticArray.h"
 #include "Containers/StringFwd.h"
-#include "Misc/StringBuilder.h"
 #include "RHI.h"
 #include "Serialization/MemoryLayout.h"
 
@@ -326,21 +325,6 @@ public:
 		return LayoutHash;	
 	}
 
-	/** Iterate recursively over all shader parameter members. */
-	template<typename TParameterFunction>
-	void IterateShaderParameterMembers(TParameterFunction Lambda) const
-	{
-		TStringBuilder<1024> ShaderBindingNameBuilder;
-		if (UseCase == EUseCase::UniformBuffer || UseCase == EUseCase::DataDrivenUniformBuffer)
-		{
-			ShaderBindingNameBuilder.Append(ShaderVariableName);
-			ShaderBindingNameBuilder.Append(TEXT("_"));
-		}
-
-		FShaderParametersMetadata::IterateShaderParameterMembersInternal(
-			*this, /* ByteOffset = */ 0, ShaderBindingNameBuilder, Lambda);
-	}
-
 	/** Iterate recursively over all FShaderParametersMetadata. */
 	template<typename TParameterFunction>
 	void IterateStructureMetadataDependencies(TParameterFunction Lambda) const
@@ -405,70 +389,4 @@ private:
 	void InitializeLayout(FRHIUniformBufferLayoutInitializer* OutLayoutInitializer = nullptr);
 
 	void AddResourceTableEntriesRecursive(const TCHAR* UniformBufferName, const TCHAR* Prefix, uint16& ResourceIndex, TMap<FString, FResourceTableEntry>& ResourceTableMap) const;
-
-	template<typename TParameterFunction>
-	static inline void IterateShaderParameterMembersInternal(
-		const FShaderParametersMetadata& ParametersMetadata,
-		uint16 ByteOffset,
-		TStringBuilder<1024>& ShaderBindingNameBuilder,
-		TParameterFunction Lambda)
-	{
-		for (const FShaderParametersMetadata::FMember& Member : ParametersMetadata.GetMembers())
-		{
-			EUniformBufferBaseType BaseType = Member.GetBaseType();
-			uint16 MemberOffset = ByteOffset + uint16(Member.GetOffset());
-			uint32 NumElements = Member.GetNumElements();
-
-			int32 MemberNameLength = FCString::Strlen(Member.GetName());
-
-			if (BaseType == UBMT_INCLUDED_STRUCT)
-			{
-				check(NumElements == 0);
-				const FShaderParametersMetadata& NewParametersMetadata = *Member.GetStructMetadata();
-				IterateShaderParameterMembersInternal(NewParametersMetadata, MemberOffset, ShaderBindingNameBuilder, Lambda);
-			}
-			else if (BaseType == UBMT_NESTED_STRUCT && NumElements == 0)
-			{
-				ShaderBindingNameBuilder.Append(Member.GetName());
-				ShaderBindingNameBuilder.Append(TEXT("_"));
-
-				const FShaderParametersMetadata& NewParametersMetadata = *Member.GetStructMetadata();
-				IterateShaderParameterMembersInternal(NewParametersMetadata, MemberOffset, ShaderBindingNameBuilder, Lambda);
-
-				ShaderBindingNameBuilder.RemoveSuffix(MemberNameLength + 1);
-			}
-			else if (BaseType == UBMT_NESTED_STRUCT && NumElements > 0)
-			{
-				ShaderBindingNameBuilder.Append(Member.GetName());
-				ShaderBindingNameBuilder.Append(TEXT("_"));
-
-				const FShaderParametersMetadata& NewParametersMetadata = *Member.GetStructMetadata();
-				for (uint32 ArrayElementId = 0; ArrayElementId < NumElements; ArrayElementId++)
-				{
-					FString ArrayElementIdString;
-					ArrayElementIdString.AppendInt(ArrayElementId);
-					int32 ArrayElementIdLength = ArrayElementIdString.Len();
-
-					ShaderBindingNameBuilder.Append(ArrayElementIdString);
-					ShaderBindingNameBuilder.Append(TEXT("_"));
-
-					uint16 NewStructOffset = MemberOffset + ArrayElementId * NewParametersMetadata.GetSize();
-					IterateShaderParameterMembersInternal(NewParametersMetadata, NewStructOffset, ShaderBindingNameBuilder, Lambda);
-
-					ShaderBindingNameBuilder.RemoveSuffix(ArrayElementIdLength + 1);
-				}
-
-				ShaderBindingNameBuilder.RemoveSuffix(MemberNameLength + 1);
-			}
-			else
-			{
-				ShaderBindingNameBuilder.Append(Member.GetName());
-
-				const TCHAR* ShaderBindingName = ShaderBindingNameBuilder.ToString();
-				Lambda(ParametersMetadata, Member, ShaderBindingName, MemberOffset);
-
-				ShaderBindingNameBuilder.RemoveSuffix(MemberNameLength);
-			}
-		}
-	}
 };
\ No newline at end of file
