Groundwork for custom shader preprocessing step that removes unused code from shader source before compilation (not used yet by the shader pipeline)

diff --git a/Engine/Source/Developer/ShaderCompilerCommon/Private/ShaderMinifier.cpp b/Engine/Source/Developer/ShaderCompilerCommon/Private/ShaderMinifier.cpp
new file mode 100644
--- /dev/null
+++ b/Engine/Source/Developer/ShaderCompilerCommon/Private/ShaderMinifier.cpp
@@ -0,0 +1,1424 @@
+// Copyright Epic Games, Inc. All Rights Reserved.
+
+#include "ShaderMinifier.h"
+
+#include "HAL/PlatformTime.h"
+#include "Hash/CityHash.h"
+#include "HlslParser.h"
+#include "Logging/LogMacros.h"
+#include "Misc/AutomationTest.h"
+#include "String/Find.h"
+
+DEFINE_LOG_CATEGORY_STATIC(LogShaderMinifier, Log, All);
+
+// TODO:
+// - track namespaces
+// - preserve multi-line #define
+// - track and re-emit correct #line
+
+namespace UE::ShaderMinifier
+{
+
+static FStringView SubStrView(FStringView S, int32 Start)
+{
+	Start = FMath::Min(Start, S.Len());
+	int32 Len = S.Len() - Start;
+	return FStringView(S.GetData() + Start, Len);
+}
+
+static FStringView SubStrView(FStringView S, int32 Start, int32 Len)
+{
+	Start = FMath::Min(Start, S.Len());
+	Len	  = FMath::Min(Len, S.Len() - Start);
+	return FStringView(S.GetData() + Start, Len);
+}
+
+template<typename TCondition>
+static FStringView SkipUntil(FStringView Source, TCondition Cond)
+{
+	int32 Cursor = 0;
+	int32 SourceLen = Source.Len();
+	while (Cursor < SourceLen)
+	{
+		if (Cond(FStringView(Source.GetData() + Cursor, SourceLen - Cursor)))
+		{
+			break;
+		}
+		++Cursor;
+	}
+	return FStringView(Source.GetData() + Cursor, SourceLen - Cursor);
+}
+
+static bool IsSpace(TCHAR C)
+{
+	switch (C)
+	{
+	default:
+		return false;
+	case TCHAR(' '):
+	case TCHAR('\f'):
+	case TCHAR('\r'):
+	case TCHAR('\n'):
+	case TCHAR('\t'):
+	case TCHAR('\v'):
+		return true;
+	}
+}
+
+static bool IsNumber(TCHAR C)
+{
+	return C >= '0' && C <= '9';
+}
+
+static bool IsPossibleIdentifierCharacter(TCHAR C)
+{
+	return (C >= '0' && C <= '9') || (C >= 'a' && C <= 'z') || (C >= 'A' && C <= 'Z') || C == '_';
+}
+
+static FStringView SkipUntilNonIdentifierCharacter(FStringView Source) {
+	int32 Len = Source.Len();
+	int32 Cursor = 0;
+	const TCHAR* SourceData = Source.GetData();
+	while (Cursor < Len)
+	{
+		if (!IsPossibleIdentifierCharacter(SourceData[Cursor]))
+		{
+			break;
+		}
+		++Cursor;
+	}
+	return FStringView(SourceData + Cursor, Len - Cursor);
+}
+
+
+static FStringView SkipSpace(FStringView Source)
+{
+	int32 Len = Source.Len();
+	int32 Cursor = 0;
+	const TCHAR* SourceData = Source.GetData();
+	while (Cursor < Len)
+	{
+		if (!IsSpace(SourceData[Cursor]))
+		{
+			break;
+		}
+		++Cursor;
+	}
+	return FStringView(SourceData + Cursor, Len - Cursor);
+}
+
+static FStringView TrimSpace(FStringView Source)
+{
+	int32 CursorBegin = 0;
+	int32 CursorEnd = Source.Len();
+
+	while (CursorBegin != CursorEnd)
+	{
+		if (!IsSpace(Source[CursorBegin]))
+		{
+			break;
+		}
+		++CursorBegin;
+	}
+
+	while (CursorBegin != CursorEnd)
+	{
+		if (!IsSpace(Source[CursorEnd-1]))
+		{
+			break;
+		}
+		--CursorEnd;
+	}
+
+	FStringView Result = SubStrView(Source, CursorBegin, CursorEnd-CursorBegin);
+
+	return Result;
+}
+
+static FStringView SkipUntilNextLine(FStringView Source)
+{
+	int32 Index = INDEX_NONE;
+	if (Source.FindChar('\n', Index))
+	{
+		Index += 1; // Skip the new line character itself
+		return FStringView(Source.GetData() + Index, Source.Len() - Index);
+	}
+	else
+	{
+		return FStringView {};
+	}
+}
+
+static FStringView SkipUntilStr(FStringView Haystack, FStringView Needle)
+{
+	return SkipUntil(Haystack, [Needle](FStringView  S) { return S.StartsWith(Needle, ESearchCase::CaseSensitive); });
+}
+
+static int32 FindFirstOf(FStringView Haystack, FStringView Needle)
+{
+	int32 Len = Haystack.Len();
+	for (int32 i = 0; i < Len; ++i)
+	{
+		TCHAR C = Haystack[i];
+		for (TCHAR C2 : Needle)
+		{
+			if (C == C2)
+			{
+				return i;
+			}
+		}
+	}
+	return INDEX_NONE;
+}
+
+static FStringView ExtractBlock(FStringView Source, TCHAR DelimBegin, TCHAR DelimEnd)
+{
+	// TODO: handle comments
+	// TODO: handle #if 0 blocks
+
+	int32 PosEnd = INDEX_NONE;
+	int32 Stack  = 0;
+	for (int32 I = 0; I < Source.Len(); ++I)
+	{
+		TCHAR C = Source[I];
+		if (C == DelimBegin)
+		{
+			Stack++;
+		}
+		else if (C == DelimEnd)
+		{
+			if (Stack == 0)
+			{
+				// delimiter mismatch
+				break;
+			}
+
+			Stack--;
+
+			if (Stack == 0)
+			{
+				PosEnd = I;
+				break;
+			}
+		}
+	}
+
+	if (Stack == 0 && PosEnd != INDEX_NONE)
+	{
+		return FStringView(Source.GetData(), PosEnd + 1);
+	}
+	else
+	{
+		return FStringView{};
+	}
+}
+
+enum class EBlockType : uint8 {
+	Unknown,
+	Keyword,	// e.g. struct, switch, register
+	Attribute,	// e.g. `[numthreads(8,8,1)]`
+	Type,		// return type of function or struct/cbuffer/variable type
+	Base,		// inheritance base type
+	Name,		// struct/variable/function name
+	Binding,	// e.g. `register(t0, space1)` or `SV_Target0`
+	Args,
+	Body,
+	Subscript,
+	TemplateArgs,
+	Expression,
+	Directive,  // #define, #pragma, #line, etc.
+	NamespaceDelimiter, // e.g. :: in an identifier like Foo::bar
+};
+
+struct FCodeBlock
+{
+	EBlockType  Type = EBlockType::Unknown;
+	FStringView Code;
+};
+
+enum class ECodeChunkType {
+	Unknown,
+	Struct,
+	CBuffer,  // HLSL cbuffer block possibly without trailing ';'
+	Function,
+	Variable,
+	Enum,
+	Define,
+	Pragma,
+};
+
+struct FCodeChunk
+{
+	ECodeChunkType Type = ECodeChunkType::Unknown;
+	TArray<FCodeBlock> Blocks;
+
+	FStringView FindFirstBlockByType(EBlockType InType) const
+	{
+		for (const FCodeBlock& Block : Blocks)
+		{
+			if (Block.Type == InType)
+			{
+				return Block.Code;
+			}
+		}
+		return {};
+	}
+};
+
+struct FParsedShader
+{
+	TArray<FCodeChunk> Chunks;
+};
+
+static FParsedShader ParseShader(FStringView InSource, FDiagnostics& Output)
+{
+	FParsedShader Result;
+
+	FStringView		   Source = InSource;
+	TArray<FCodeBlock> PendingBlocks;
+	TArray<FCodeChunk> Chunks;
+
+	ECodeChunkType ChunkType = ECodeChunkType::Unknown;
+	bool		   bFoundBody = false;
+	bool		   bFoundColon = false;
+	bool		   bFoundIdentifier = false;
+	bool		   bFoundAssignment = false;
+
+	int32 ArgsBlockIndex = INDEX_NONE;
+	int32 CbufferBlockIndex = INDEX_NONE;
+	int32 StructBlockIndex = INDEX_NONE;
+	int32 EnumBlockIndex = INDEX_NONE;
+	int32 BodyBlockIndex = INDEX_NONE;
+	int32 ExpressionBlockIndex = INDEX_NONE;
+
+	auto AddDiagnostic = [InSource, &Source](TArray<FDiagnosticMessage>& Output, FStringView Message)
+	{
+		FDiagnosticMessage Diagnostic;
+
+		Diagnostic.Message = FString(Message);
+		Diagnostic.Offset = int32(Source.GetData() - InSource.GetData());
+		// Diagnostic.Line = ...; // TODO
+		// Diagnostic.Column = ...; // TODO
+
+		Output.Add(MoveTemp(Diagnostic));
+	};
+
+	auto AddBlock = [&PendingBlocks](EBlockType Type, FStringView Code)
+	{
+		FCodeBlock NewBlock;
+		NewBlock.Type = Type;
+		NewBlock.Code = Code;
+		PendingBlocks.Push(NewBlock);
+	};
+
+	auto FinalizeChunk = [&]() {
+		const bool bFoundArgs = ArgsBlockIndex >= 0;
+
+		bool bHasType = false;
+		bool bHasName = false;
+
+		if (!PendingBlocks.IsEmpty())
+		{
+			if (ChunkType == ECodeChunkType::Unknown)
+			{
+				if (bFoundIdentifier && bFoundArgs && bFoundBody)
+				{
+					ChunkType = ECodeChunkType::Function;
+				}
+				else if (bFoundIdentifier)
+				{
+					ChunkType = ECodeChunkType::Variable;
+				}
+			}
+
+			int32 NameBlockIndex = INDEX_NONE;
+
+			if (ChunkType == ECodeChunkType::Struct)
+			{
+				check(StructBlockIndex >= 0);
+
+				PendingBlocks[StructBlockIndex].Type = EBlockType::Keyword;
+
+				int32 TypeBlockIndex = StructBlockIndex + 1;
+				if (TypeBlockIndex != BodyBlockIndex && TypeBlockIndex < PendingBlocks.Num())
+				{
+					PendingBlocks[TypeBlockIndex].Type = EBlockType::Type;
+					bHasType = true;
+				}
+
+				if (ExpressionBlockIndex > 0)
+				{
+					NameBlockIndex = ExpressionBlockIndex - 1;
+					PendingBlocks[NameBlockIndex].Type = EBlockType::Name;
+					bHasName = true;
+				}
+			}
+			else if (ChunkType == ECodeChunkType::CBuffer)
+			{
+				check(CbufferBlockIndex >= 0);
+
+				PendingBlocks[CbufferBlockIndex].Type = EBlockType::Keyword;
+
+				int32 TypeBlockIndex = CbufferBlockIndex + 1;
+				if (TypeBlockIndex != BodyBlockIndex && TypeBlockIndex < PendingBlocks.Num())
+				{
+					PendingBlocks[TypeBlockIndex].Type = EBlockType::Type;
+				}
+			}
+			else if (ChunkType == ECodeChunkType::Enum)
+			{
+				check(EnumBlockIndex >= 0);
+
+				PendingBlocks[EnumBlockIndex].Type = EBlockType::Keyword;
+
+				if (BodyBlockIndex > 1)
+				{
+					PendingBlocks[BodyBlockIndex - 1].Type = EBlockType::Type;
+				}
+			}
+			else if (ChunkType == ECodeChunkType::Function)
+			{
+				NameBlockIndex = ArgsBlockIndex - 1;
+				if (NameBlockIndex >= 0)
+				{
+					PendingBlocks[NameBlockIndex].Type = EBlockType::Name;
+				}
+			}
+			else if (ChunkType == ECodeChunkType::Variable)
+			{
+				// TODO: tag name / type / binding
+			}
+
+			if (ChunkType == ECodeChunkType::Struct && bHasName && !bHasType)
+			{
+				ChunkType = ECodeChunkType::Variable;
+			}
+
+			if (ChunkType == ECodeChunkType::Struct && bHasName && bHasType)
+			{
+				// Handle simultaneous struct type and variable declaration
+
+				FCodeChunk StructChunk;
+				StructChunk.Type = ECodeChunkType::Struct;
+
+				for (int32 i = int32(StructBlockIndex); i < NameBlockIndex; ++i)
+				{
+					StructChunk.Blocks.Push(PendingBlocks[i]);
+				}
+
+				FCodeChunk VarChunk;
+				VarChunk.Type = ECodeChunkType::Variable;
+
+				for (int32 i = 0; i < PendingBlocks.Num(); ++i)
+				{
+					if (i == StructBlockIndex || i == BodyBlockIndex)
+					{
+						continue;
+					}
+					VarChunk.Blocks.Push(PendingBlocks[i]);
+				}
+
+				Chunks.Push(StructChunk);
+				Chunks.Push(VarChunk);
+			}
+			else
+			{
+				FCodeChunk Chunk;
+				Chunk.Type = ChunkType;
+				Swap(Chunk.Blocks, PendingBlocks);
+				Chunks.Push(Chunk);
+			}
+
+			ChunkType = ECodeChunkType::Unknown;
+			ArgsBlockIndex = INDEX_NONE;
+			CbufferBlockIndex = INDEX_NONE;
+			StructBlockIndex = INDEX_NONE;
+			EnumBlockIndex = INDEX_NONE;
+			BodyBlockIndex = INDEX_NONE;
+			ExpressionBlockIndex = INDEX_NONE;
+			bFoundBody = false;
+			bFoundColon = false;
+			bFoundIdentifier = false;
+			bFoundAssignment = false;
+
+			PendingBlocks.Empty();
+		}
+	};
+
+	while (Output.Errors.IsEmpty())
+	{
+		Source = SkipSpace(Source);
+
+		if (Source.IsEmpty())
+		{
+			break;
+		}
+
+		if (Source.StartsWith(TEXT("//")) || Source.StartsWith(TEXT("#line")))
+		{
+			Source = SkipUntilNextLine(Source);
+			continue;
+		}
+		else if (Source.StartsWith(TEXT("#pragma")))
+		{
+			FStringView Remainder = SkipUntilNextLine(Source);
+			FStringView Block = SubStrView(Source, 0, Source.Len() - Remainder.Len());
+			AddBlock(EBlockType::Directive, Block);
+			ChunkType = ECodeChunkType::Pragma;
+			FinalizeChunk();
+			Source = Remainder;
+			continue;
+		}
+		else if (Source.StartsWith(TEXT("#define")))
+		{
+			// TODO: handle `\` new lines in defines
+			FStringView Remainder = SkipUntilNextLine(Source);
+			FStringView Block = SubStrView(Source, 0, Source.Len() - Remainder.Len());
+			AddBlock(EBlockType::Directive, Block);
+			ChunkType = ECodeChunkType::Define;
+			FinalizeChunk();
+			Source = Remainder;
+			continue;
+		}
+		else if (Source.StartsWith(TEXT("#if 0")))
+		{
+			Source = SkipUntilStr(Source, TEXT("#endif"));
+			if (Source.Len() >= 6)
+			{
+				Source = SubStrView(Source, 6);
+			}
+			continue;
+		}
+		else if (Source.StartsWith(TEXT("/*")))
+		{
+			Source = SkipUntilStr(Source, TEXT("*/"));
+			if (Source.Len() >= 2)
+			{
+				Source = SubStrView(Source, 2);
+			}
+			continue;
+		}
+
+		FStringView Remainder = SkipUntilNonIdentifierCharacter(Source);
+
+		FStringView Identifier = SubStrView(Source, 0, Source.Len() - Remainder.Len());
+
+		if (Identifier.Len())
+		{
+			if (ChunkType == ECodeChunkType::Unknown)
+			{
+				if (Identifier == TEXT("struct"))
+				{
+					ChunkType = ECodeChunkType::Struct;
+					StructBlockIndex = PendingBlocks.Num();
+				}
+				else if (Identifier == TEXT("cbuffer") || Identifier == TEXT("ConstantBuffer"))
+				{
+					ChunkType = ECodeChunkType::CBuffer;
+					CbufferBlockIndex = PendingBlocks.Num();
+				}
+				else if (Identifier == TEXT("enum"))
+				{
+					ChunkType = ECodeChunkType::Enum;
+					EnumBlockIndex = PendingBlocks.Num();
+				}
+			}
+
+			EBlockType BlockType = EBlockType::Unknown;
+
+			if (bFoundColon)
+			{
+				if (ChunkType == ECodeChunkType::Struct)
+				{
+					BlockType = EBlockType::Base;
+				}
+				else
+				{
+					BlockType = EBlockType::Binding;
+				}
+
+				bFoundColon = false;
+			}
+
+			AddBlock(BlockType, Identifier);
+
+			Source = Remainder;
+			bFoundIdentifier = true;
+
+			continue;
+		}
+
+		FStringView Block;
+
+		TCHAR C = Source[0];
+
+		EBlockType BlockType = EBlockType::Unknown;
+
+		if (Source.StartsWith(TEXT("==")))
+		{
+			AddDiagnostic(Output.Errors, TEXT("Unexpected sequence '=='"));
+			break;
+		}
+		else if (Source.StartsWith(TEXT("::")))
+		{
+			Block = SubStrView(Source, 0, 2);
+			Source = SubStrView(Source, 2);
+
+			AddBlock(EBlockType::NamespaceDelimiter, Block);
+
+			continue;
+		}
+		else if (C == '=')
+		{
+			int32 Pos = FindFirstOf(Source, TEXT("{;"));
+
+			if (Pos == INDEX_NONE)
+			{
+				AddDiagnostic(Output.Errors, TEXT("Expected block body or semicolon after '='"));
+				break;
+			}
+
+			bFoundAssignment = true;
+
+			char C2 = Source[Pos];
+
+			if (C2 == '{')
+			{
+				Source = SubStrView(Source, Pos);
+				continue;
+			}
+			else if (C2 == ';')
+			{
+				Block = SubStrView(Source, 1, Pos - 1);
+				Block = TrimSpace(Block);
+
+				int32 BlockOffset = int32(Block.GetData() - Source.GetData());
+				Source = SubStrView(Source, BlockOffset);
+
+				BlockType = EBlockType::Expression;
+				ExpressionBlockIndex = int32(PendingBlocks.Num());
+			}
+		}
+		else if (C == ':')
+		{
+			Source = SubStrView(Source, 1);
+			bFoundColon = true;
+			continue;
+		}
+		else if (C == '(')
+		{
+			Block = ExtractBlock(Source, '(', ')');
+			BlockType = EBlockType::Args;
+
+			if (ArgsBlockIndex < 0)
+			{
+				ArgsBlockIndex = PendingBlocks.Num();
+			}
+		}
+		else if (C == '{')
+		{
+			Block = ExtractBlock(Source, '{', '}');
+
+			if (BodyBlockIndex == INDEX_NONE && !bFoundAssignment)
+			{
+				BlockType = EBlockType::Body;
+				BodyBlockIndex = PendingBlocks.Num();
+				bFoundBody = true;
+			}
+			else if (bFoundAssignment)
+			{
+				BlockType = EBlockType::Expression;
+				ExpressionBlockIndex = PendingBlocks.Num();
+			}
+		}
+		else if (C == '[')
+		{
+			Block = ExtractBlock(Source, '[', ']');
+
+			if (bFoundIdentifier)
+			{
+				BlockType = EBlockType::Subscript;
+			}
+			else
+			{
+				BlockType = EBlockType::Attribute;
+			}
+		}
+		else if (C == '<')
+		{
+			Block = ExtractBlock(Source, '<', '>');
+			BlockType = EBlockType::TemplateArgs;
+
+			if (ChunkType == ECodeChunkType::CBuffer)
+			{
+				// `ConstantBuffer<Foo>` is treated as a variable/resource declaration rather than a cbuffer block
+				ChunkType = ECodeChunkType::Variable;
+			}
+		}
+		else if (C == ';')
+		{
+			FinalizeChunk();
+			Source = SubStrView(Source, 1);
+			continue;
+		}
+		else
+		{
+			AddDiagnostic(Output.Errors, TEXT("Unexpected character"));
+			break;
+		}
+
+		if (Block.IsEmpty())
+		{
+			AddDiagnostic(Output.Errors, TEXT("Failed to extract code block"));
+			break;
+		}
+		else
+		{
+			AddBlock(BlockType, Block);
+
+			Source = SubStrView(Source, Block.Len());
+
+			if (BlockType == EBlockType::Body && ArgsBlockIndex != INDEX_NONE)
+			{
+				FinalizeChunk();
+			}
+			else if (BlockType == EBlockType::Body && CbufferBlockIndex != INDEX_NONE)
+			{
+				FinalizeChunk();
+			}
+			else if (BlockType == EBlockType::Expression)
+			{
+				FinalizeChunk();
+			}
+		}
+	}
+
+	std::swap(Result.Chunks, Chunks);
+
+	return Result;
+}
+
+template<typename CallbackT>
+void FindChunksByIdentifier(TConstArrayView<FCodeChunk> Chunks, FStringView Identifier, CallbackT Callback)
+{
+	for (const FCodeChunk& Chunk : Chunks)
+	{
+		for (const FCodeBlock& Block : Chunk.Blocks)
+		{
+			if (Block.Code == Identifier)
+			{
+				Callback(Chunk);
+			}
+		}
+	}
+}
+
+static TArray<FStringView> SplitByChar(FStringView Source, TCHAR Delimiter)
+{
+	TArray<FStringView> Result;
+
+	int32 Start = 0;
+
+	for (int32 I = 0; I < Source.Len(); ++I)
+	{
+		TCHAR C = Source[I];
+		if (C == Delimiter)
+		{
+			size_t Len = I - Start;
+			Result.Push(SubStrView(Source, Start, Len));
+			Start = I + 1;
+		}
+	}
+
+	if (Start != Source.Len())
+	{
+		int32 Len = Source.Len() - Start;
+		Result.Push(SubStrView(Source, Start, Len));
+	}
+
+	return Result;
+}
+
+static void ExtractIdentifiers(FStringView InSource, TArray<FStringView>& Result)
+{
+	FStringView Source = InSource;
+
+	const ESearchCase::Type SC = ESearchCase::CaseSensitive;
+
+	for (;;)
+	{
+		Source = SkipSpace(Source);
+
+		if (Source.IsEmpty())
+		{
+			break;
+		}
+
+		if (!IsPossibleIdentifierCharacter(Source[0]))
+		{
+			if (Source.StartsWith(TEXT("//"), SC) 
+				|| Source.StartsWith(TEXT("#line"), SC) 
+				|| Source.StartsWith(TEXT("#pragma"), SC))
+			{
+				Source = SkipUntilNextLine(Source);
+				continue;
+			}
+			else if (Source.StartsWith(TEXT("#if 0"), SC))
+			{
+				Source = SkipUntilStr(Source, TEXT("#endif"));
+				if (Source.Len() >= 6)
+				{
+					Source = SubStrView(Source, 6);
+				}
+				continue;
+			}
+			else if (Source.StartsWith(TEXT("/*"), SC))
+			{
+				Source = SkipUntilStr(Source, TEXT("*/"));
+				if (Source.Len() >= 2)
+				{
+					Source = SubStrView(Source, 2);
+				}
+				continue;
+			}
+		}
+
+		FStringView Remainder = SkipUntilNonIdentifierCharacter(Source);
+
+		FStringView Identifier = SubStrView(Source, 0, Source.Len() - Remainder.Len());
+
+		if (Identifier.IsEmpty())
+		{
+			if (!Remainder.IsEmpty())
+			{
+				Remainder = SubStrView(Remainder, 1);
+			}
+		}
+		else
+		{
+			if (!IsNumber(Identifier[0]))  // Identifiers can't start with numbers
+			{
+				Result.Push(Identifier);
+			}
+		}
+
+		Source = Remainder;
+	}
+}
+
+static void ExtractIdentifiers(const FCodeChunk& Chunk, TArray<FStringView>& Result)
+{
+	for (const FCodeBlock& Block : Chunk.Blocks)
+	{
+		ExtractIdentifiers(Block.Code, Result);
+	}
+}
+
+static void OutputChunk(const FCodeChunk& Chunk, FStringBuilderBase& OutputStream)
+{
+	if (Chunk.Blocks.IsEmpty())
+	{
+		return;
+	}
+
+	int32 Index = 0;
+	for (const FCodeBlock& Block : Chunk.Blocks)
+	{
+		if (Index != 0)
+		{
+			OutputStream << ' ';
+		}
+
+		if (Block.Type == EBlockType::Expression)
+		{
+			OutputStream << "= ";
+		}
+		else if (Block.Type == EBlockType::Body)
+		{
+			OutputStream << "\n";
+		}
+
+		if (Block.Type == EBlockType::Binding || Block.Type == EBlockType::Base)
+		{
+			OutputStream << ": ";
+		}
+
+		OutputStream << Block.Code;
+
+		++Index;
+	}
+
+	if (Chunk.Type != ECodeChunkType::Function
+		&& Chunk.Type != ECodeChunkType::CBuffer
+		&& Chunk.Type != ECodeChunkType::Pragma
+		&& Chunk.Type != ECodeChunkType::Define)
+	{
+		OutputStream << ";";
+	}
+
+	OutputStream << "\n\n";
+}
+
+struct FCasedStringViewKeyFuncs : public DefaultKeyFuncs<FStringView>
+{
+	static FORCEINLINE FStringView GetSetKey(FStringView K) { return K; }
+	template <typename T>
+	static FORCEINLINE FStringView GetSetKey(const TPair<FStringView, T>& P) { return P.Key; }
+	static FORCEINLINE bool Matches(FStringView A, FStringView B) { return A.Equals(B, ESearchCase::CaseSensitive); }
+	static FORCEINLINE uint32 GetKeyHash(FStringView Key)
+	{
+		return CityHash32((const char*)Key.GetData(), Key.Len() * sizeof(*Key.GetData()));
+	}
+};
+
+static FString MinifyShader(const FParsedShader& Parsed, FStringView SemicolonSeparatedEntryPoints, EMinifyShaderFlags Flags, FDiagnostics& Diagnostics)
+{
+	FStringBuilderBase OutputStream;
+
+	TSet<FStringView, FCasedStringViewKeyFuncs, FDefaultSetAllocator> RelevantIdentifiers;
+
+	TSet<const FCodeChunk*> RelevantChunks;
+	TSet<FStringView, FCasedStringViewKeyFuncs, FDefaultSetAllocator>  ProcessedIdentifiers;
+
+	TArray<const FCodeChunk*> PendingChunks;
+
+	TArray<FStringView> EntryPoints = SplitByChar(SemicolonSeparatedEntryPoints, ';');
+
+	for (FStringView Entry : EntryPoints)
+	{
+		RelevantIdentifiers.Add(Entry);
+		ProcessedIdentifiers.Add(Entry);
+		FindChunksByIdentifier(Parsed.Chunks, Entry, [&PendingChunks](const FCodeChunk& Chunk) { PendingChunks.Push(&Chunk); });
+	}
+
+	for (const FCodeChunk* Chunk : PendingChunks)
+	{
+		RelevantChunks.Add(Chunk);
+	}
+
+	{
+		// Some known builtin words to ignore
+
+		ProcessedIdentifiers.Add(TEXT("asfloat"));
+		ProcessedIdentifiers.Add(TEXT("asint"));
+		ProcessedIdentifiers.Add(TEXT("asuint"));
+		ProcessedIdentifiers.Add(TEXT("bool"));
+		ProcessedIdentifiers.Add(TEXT("bool2"));
+		ProcessedIdentifiers.Add(TEXT("bool3"));
+		ProcessedIdentifiers.Add(TEXT("bool4"));
+		ProcessedIdentifiers.Add(TEXT("cbuffer"));
+		ProcessedIdentifiers.Add(TEXT("const"));
+		ProcessedIdentifiers.Add(TEXT("else"));
+		ProcessedIdentifiers.Add(TEXT("false"));
+		ProcessedIdentifiers.Add(TEXT("float"));
+		ProcessedIdentifiers.Add(TEXT("float2"));
+		ProcessedIdentifiers.Add(TEXT("float3"));
+		ProcessedIdentifiers.Add(TEXT("float3x3"));
+		ProcessedIdentifiers.Add(TEXT("float3x4"));
+		ProcessedIdentifiers.Add(TEXT("float4"));
+		ProcessedIdentifiers.Add(TEXT("float4x4"));
+		ProcessedIdentifiers.Add(TEXT("for"));
+		ProcessedIdentifiers.Add(TEXT("if"));
+		ProcessedIdentifiers.Add(TEXT("in"));
+		ProcessedIdentifiers.Add(TEXT("inout"));
+		ProcessedIdentifiers.Add(TEXT("int"));
+		ProcessedIdentifiers.Add(TEXT("int2"));
+		ProcessedIdentifiers.Add(TEXT("int3"));
+		ProcessedIdentifiers.Add(TEXT("int4"));
+		ProcessedIdentifiers.Add(TEXT("out"));
+		ProcessedIdentifiers.Add(TEXT("return"));
+		ProcessedIdentifiers.Add(TEXT("static"));
+		ProcessedIdentifiers.Add(TEXT("struct"));
+		ProcessedIdentifiers.Add(TEXT("switch"));
+		ProcessedIdentifiers.Add(TEXT("true"));
+		ProcessedIdentifiers.Add(TEXT("uint"));
+		ProcessedIdentifiers.Add(TEXT("uint2"));
+		ProcessedIdentifiers.Add(TEXT("uint3"));
+		ProcessedIdentifiers.Add(TEXT("uint4"));
+		ProcessedIdentifiers.Add(TEXT("void"));
+		ProcessedIdentifiers.Add(TEXT("while"));
+
+#if 0 // Some shaders define template versions of some built-in functions, so we can't trivially ignore them
+		//ProcessedIdentifiers.Add(TEXT("abs"));
+		//ProcessedIdentifiers.Add(TEXT("any"));
+		//ProcessedIdentifiers.Add(TEXT("clamp"));
+		//ProcessedIdentifiers.Add(TEXT("clip"));
+		//ProcessedIdentifiers.Add(TEXT("cos"));
+		//ProcessedIdentifiers.Add(TEXT("cross"));
+		//ProcessedIdentifiers.Add(TEXT("dot"));
+		//ProcessedIdentifiers.Add(TEXT("frac"));
+		//ProcessedIdentifiers.Add(TEXT("lerp"));
+		//ProcessedIdentifiers.Add(TEXT("max"));
+		//ProcessedIdentifiers.Add(TEXT("min"));
+		//ProcessedIdentifiers.Add(TEXT("mul"));
+		//ProcessedIdentifiers.Add(TEXT("normalize"));
+		//ProcessedIdentifiers.Add(TEXT("pow"));
+		//ProcessedIdentifiers.Add(TEXT("saturate"));
+		//ProcessedIdentifiers.Add(TEXT("sign"));
+		//ProcessedIdentifiers.Add(TEXT("sin"));
+		//ProcessedIdentifiers.Add(TEXT("sqrt"));
+#endif
+	}
+
+	if (PendingChunks.IsEmpty())
+	{
+		// printf("Entry point chunk is not found in the shader\n");
+		return {};
+	}
+
+	TArray<FStringView> TempIdentifiers;
+
+	TMap<FStringView, TArray<const FCodeChunk*>, FDefaultSetAllocator, FCasedStringViewKeyFuncs> ChunksByIdentifier;
+	for (const FCodeChunk& Chunk : Parsed.Chunks)
+	{
+		for (const FCodeBlock& Block : Chunk.Blocks)
+		{
+			if (Chunk.Type == ECodeChunkType::Function && Block.Type != EBlockType::Name)
+			{
+				continue;
+			}
+
+			if (Chunk.Type == ECodeChunkType::Struct && Block.Type != EBlockType::Type)
+			{
+				continue;
+			}
+
+			if (Chunk.Type == ECodeChunkType::CBuffer && Block.Type == EBlockType::Body)
+			{
+				TempIdentifiers.Empty();
+				ExtractIdentifiers(Block.Code, TempIdentifiers);
+				for (FStringView Identifier : TempIdentifiers)
+				{
+					ChunksByIdentifier.FindOrAdd(Identifier).Push(&Chunk);
+				}
+
+				continue;
+			}
+
+			ChunksByIdentifier.FindOrAdd(Block.Code).Push(&Chunk);
+		}
+	}
+
+	TMap<const FCodeChunk*, const FCodeChunk*> ChunkRequestedBy;
+
+	while (!PendingChunks.IsEmpty())
+	{
+		TempIdentifiers.Empty();
+
+		const FCodeChunk* CurrentChunk = PendingChunks.Last();
+		PendingChunks.Pop();
+
+		ExtractIdentifiers(*CurrentChunk, TempIdentifiers);
+
+		for (FStringView Identifier : TempIdentifiers)
+		{
+			bool bIdentifierWasAlreadyInSet = false;
+			ProcessedIdentifiers.Add(Identifier, &bIdentifierWasAlreadyInSet);
+			if (!bIdentifierWasAlreadyInSet)
+			{
+				auto FoundChunks = ChunksByIdentifier.Find(Identifier);
+				if (FoundChunks == nullptr)
+				{
+					continue;
+				}
+				for (const FCodeChunk* Chunk : *FoundChunks)
+				{
+					if (Chunk == CurrentChunk)
+					{
+						continue;
+					}
+
+					bool bChunkWasAlreadyInSet = false;
+					RelevantChunks.Add(Chunk, &bChunkWasAlreadyInSet);
+					if (!bChunkWasAlreadyInSet)
+					{
+						PendingChunks.Push(Chunk);
+
+						// TODO: handle cbuffers, resources and globals
+						if (Chunk->Type == ECodeChunkType::Function || Chunk->Type == ECodeChunkType::Struct)
+						{
+							ChunkRequestedBy.FindOrAdd(Chunk) = CurrentChunk;
+						}
+					}
+				}
+			}
+		}
+	}
+
+	uint32 NumFunctions		= 0;
+	uint32 NumStructs		= 0;
+	uint32 NumVariables		= 0;
+	uint32 NumCBuffers		= 0;
+	uint32 NumOtherChunks	= 0;
+
+	for (const FCodeChunk& Chunk : Parsed.Chunks)
+	{
+		if (RelevantChunks.Find(&Chunk) == nullptr)
+		{
+			continue;
+		}
+
+		// TODO: print chunk info
+		if (Chunk.Type == ECodeChunkType::Function)
+		{
+			NumFunctions += 1;
+		}
+		else if (Chunk.Type == ECodeChunkType::Struct)
+		{
+			NumStructs += 1;
+		}
+		else if (Chunk.Type == ECodeChunkType::Variable)
+		{
+			NumVariables += 1;
+		}
+		else if (Chunk.Type == ECodeChunkType::CBuffer)
+		{
+			NumCBuffers += 1;
+		}
+		else
+		{
+			NumOtherChunks += 1;
+		}
+	}
+
+	if (EnumHasAnyFlags(Flags, EMinifyShaderFlags::OutputStats))
+	{
+		OutputStream << "// Total code chunks: " << RelevantChunks.Num() << "\n";
+		OutputStream << "// - Functions: " << NumFunctions << "\n";
+		OutputStream << "// - Structs: " << NumStructs << "\n";
+		OutputStream << "// - CBuffers: " << NumCBuffers << "\n";
+		OutputStream << "// - Variables: " << NumVariables << "\n";
+		OutputStream << "// - Other: " << NumOtherChunks << "\n";
+		OutputStream << "\n";
+	}
+
+	for (const FCodeChunk& Chunk : Parsed.Chunks)
+	{
+		if (Chunk.Type != ECodeChunkType::Pragma      // Pragmas and defines that remain after preprocessing
+			&& Chunk.Type != ECodeChunkType::Define   // must be preserved as they may control important compiler behaviors.
+			// && Chunk.Type != ECodeChunkType::Variable // Varibles must be currently preserved due to how RootParametersStructure is generated later.
+			// && Chunk.Type != ECodeChunkType::CBuffer  // TODO: analyze cbuffer member usage
+			&& RelevantChunks.Find(&Chunk) == nullptr)
+		{
+			continue;
+		}
+
+		if (EnumHasAnyFlags(Flags, EMinifyShaderFlags::OutputReasons))
+		{
+			auto RequestedBy = ChunkRequestedBy.Find(&Chunk);
+			if (RequestedBy != nullptr)
+			{
+				const FCodeChunk* RequestedByChunk = *RequestedBy;
+				FStringView  RequestedByName  = RequestedByChunk->FindFirstBlockByType(EBlockType::Name);
+				if (!RequestedByName.IsEmpty())
+				{
+					OutputStream << "// REASON: " << RequestedByName << "\n";
+				}
+			}
+		}
+
+		OutputChunk(Chunk, OutputStream);
+	}
+
+	FString Output = FString(OutputStream.ToView());
+
+	return Output;
+}
+
+FMinifiedShader Minify(const FStringView PreprocessedShader, const FStringView EntryPoint, EMinifyShaderFlags Flags)
+{
+	FMinifiedShader Result;
+
+	FParsedShader Parsed = ParseShader(PreprocessedShader, Result.Diagnostics);
+
+	if (!Parsed.Chunks.IsEmpty())
+	{
+		Result.Code = MinifyShader(Parsed, EntryPoint, Flags, Result.Diagnostics);
+	}
+
+	return Result;
+}
+
+} // namespace UE::ShaderMinifier
+
+#if WITH_AUTOMATION_TESTS
+
+IMPLEMENT_SIMPLE_AUTOMATION_TEST(FShaderMinifierParserTest, "System.Shaders.ShaderMinifier.Parse", EAutomationTestFlags::EditorContext | EAutomationTestFlags::EngineFilter);
+
+namespace UE::ShaderMinifier
+{
+// Convenience wrapper for tests where we don't care about diagnostic messages
+static FParsedShader ParseShader(FStringView InSource)
+{
+	FDiagnostics Diagnostics;
+	return ParseShader(InSource, Diagnostics);
+}
+}
+
+bool FShaderMinifierParserTest::RunTest(const FString& Parameters)
+{
+	using namespace UE::ShaderMinifier;
+
+	TestEqual(TEXT("SkipSpace"), 
+		FString(SkipSpace(TEXT("  \n\r\f \tHello"))), 
+		FString(TEXT("Hello")));
+
+	TestEqual(TEXT("SkipUntilStr (found)"), 
+		FString(SkipUntilStr(TEXT("Hello World"), TEXT("World"))),
+		FString(TEXT("World")));
+
+	TestEqual(TEXT("SkipUntilStr (not found)"),
+		FString(SkipUntilStr(TEXT("Hello World"), TEXT("Blah"))),
+		FString());
+
+	{
+		auto P = ParseShader(TEXT("static const struct { int Blah; } Foo = { 123; };"));
+		TestEqual(TEXT("Anonymous struct variable with initializer, total chunks"), P.Chunks.Num(), 1);
+		if (P.Chunks.Num() == 1)
+		{
+			TestEqual(TEXT("Anonymous struct variable with initializer, main chunk type"), P.Chunks[0].Type, ECodeChunkType::Variable);
+		}
+	}
+
+	{
+		auto P = ParseShader(TEXT("float4 PSMain() : SV_Target { return float4(1,0,0,1); };"));
+		TestEqual(TEXT("Pixel shader entry point, total chunks"), P.Chunks.Num(), 1);
+		if (P.Chunks.Num() == 1)
+		{
+			TestEqual(TEXT("Pixel shader entry point, main chunk type"), P.Chunks[0].Type, ECodeChunkType::Function);
+		}
+	}
+
+	{
+		TArray<FStringView> R;
+		ExtractIdentifiers(TEXT("Hello[World]; Foo[0];\n"), R);
+		if (TestEqual(TEXT("ExtractIdentifiers1: Num"), R.Num(), 3))
+		{
+			TestEqual(TEXT("ExtractIdentifiers1: R[0]"), FString(R[0]), TEXT("Hello"));
+			TestEqual(TEXT("ExtractIdentifiers1: R[1]"), FString(R[1]), TEXT("World"));
+			TestEqual(TEXT("ExtractIdentifiers1: R[2]"), FString(R[2]), TEXT("Foo"));
+		}
+	}
+
+	{
+		TArray<FStringView> R;
+		ExtractIdentifiers(TEXT("#line 0\nStructuredBuffer<uint4> Blah : register(t0, space123);#line 1\n#pragma foo\n"), R);
+		if (TestEqual(TEXT("ExtractIdentifiers2: Num"), R.Num(), 6))
+		{
+			TestEqual(TEXT("ExtractIdentifiers2: R[0]"), FString(R[0]), TEXT("StructuredBuffer"));
+			TestEqual(TEXT("ExtractIdentifiers2: R[1]"), FString(R[1]), TEXT("uint4"));
+			TestEqual(TEXT("ExtractIdentifiers2: R[2]"), FString(R[2]), TEXT("Blah"));
+			TestEqual(TEXT("ExtractIdentifiers2: R[3]"), FString(R[3]), TEXT("register"));
+			TestEqual(TEXT("ExtractIdentifiers2: R[4]"), FString(R[4]), TEXT("t0"));
+			TestEqual(TEXT("ExtractIdentifiers2: R[5]"), FString(R[5]), TEXT("space123"));
+		}
+	}
+
+	{
+		auto P = ParseShader(TEXT("StructuredBuffer<uint4> Blah : register(t0, space123);"));
+		if (TestEqual(TEXT("ParseShader: structured buffer: num chunks"), P.Chunks.Num(), 1))
+		{
+			TestEqual(TEXT("ParseShader: structured buffer: chunk"), P.Chunks[0].Type, ECodeChunkType::Variable);
+		}
+	}
+
+	{
+		auto P = ParseShader(TEXT("const float Foo = 123.45f;"));
+		if (TestEqual(TEXT("ParseShader: const float with initializer: num chunks"), P.Chunks.Num(), 1))
+		{
+			TestEqual(TEXT("ParseShader: const float with initializer: chunk type"), P.Chunks[0].Type, ECodeChunkType::Variable);
+		}
+	}
+
+	{
+		auto P = ParseShader(TEXT("struct Blah { int A; };"));
+		if (TestEqual(TEXT("ParseShader: struct: num chunks"), P.Chunks.Num(), 1))
+		{
+			TestEqual(TEXT("ParseShader: struct: chunk type"), P.Chunks[0].Type, ECodeChunkType::Struct);
+		}
+	}
+
+	{
+		auto P = ParseShader(TEXT("struct Foo { int FooA; }; struct Bar : Foo { int BarA; };"));
+		if (TestEqual(TEXT("ParseShader: inherited struct: num chunks"), P.Chunks.Num(), 2))
+		{
+			TestEqual(TEXT("ParseShader: inherited struct: chunk 0 type"), P.Chunks[0].Type, ECodeChunkType::Struct);
+			TestEqual(TEXT("ParseShader: inherited struct: chunk 1 type"), P.Chunks[1].Type, ECodeChunkType::Struct);
+		}
+	}
+
+	{
+		auto P = ParseShader(TEXT("[numthreads(8,8,1)] void Main() {};"));
+		if (TestEqual(TEXT("ParseShader: compute shader entry point: num chunks"), P.Chunks.Num(), 1))
+		{
+			TestEqual(TEXT("ParseShader: compute shader entry point: chunk type"), P.Chunks[0].Type, ECodeChunkType::Function);
+
+			if (TestEqual(TEXT("ParseShader: compute shader entry point: num blocks"), P.Chunks[0].Blocks.Num(), 5))
+			{
+				TestEqual(TEXT("ParseShader: compute shader entry point: attribute block type"), P.Chunks[0].Blocks[0].Type, EBlockType::Attribute);
+			}
+		}
+	}
+
+	{
+		auto P = ParseShader(TEXT("Texture2D Blah : register(t0);"));
+		if (TestEqual(TEXT("ParseShader: texture with register: num chunks"), P.Chunks.Num(), 1))
+		{
+			TestEqual(TEXT("ParseShader: texture with register: chunk type"), P.Chunks[0].Type, ECodeChunkType::Variable);
+		}
+	}
+	{
+		auto P = ParseShader(TEXT("Texture2D Blah;"));
+		if (TestEqual(TEXT("ParseShader: texture: num chunks"), P.Chunks.Num(), 1))
+		{
+			TestEqual(TEXT("ParseShader: texture: chunk type"), P.Chunks[0].Type, ECodeChunkType::Variable);
+		}
+	}
+
+	{
+		auto P = ParseShader(TEXT("SamplerState Blah : register(s0, space123);"));
+		if (TestEqual(TEXT("ParseShader: sampler state with register: num chunks"), P.Chunks.Num(), 1))
+		{
+			TestEqual(TEXT("ParseShader: sampler state with register: chunk type"), P.Chunks[0].Type, ECodeChunkType::Variable);
+		}
+	}
+
+#if 0
+	{
+		// TODO: handle function forward declarations
+		auto P = ParseShader(TEXT("Foo Fun(int a);"));
+		TestEqual(TEXT("ParseShader: function forward declaration"), P.Chunks[0].Type, ECodeChunkType::FunctionDecl);
+	}
+#endif
+
+	{
+		auto P = ParseShader(TEXT("void Fun(int a) {};"));
+		if (TestEqual(TEXT("ParseShader: function with trailing semicolon: num chunks"), P.Chunks.Num(), 1))
+		{
+			TestEqual(TEXT("ParseShader: function with trailing semicolon: chunk type"), P.Chunks[0].Type, ECodeChunkType::Function);
+		}
+	}
+
+	{
+		auto P = ParseShader(TEXT("void Fun(int a) {}"));
+		if (TestEqual(TEXT("ParseShader: function: num chunks"), P.Chunks.Num(), 1))
+		{
+			TestEqual(TEXT("ParseShader: function: chunk type"), P.Chunks[0].Type, ECodeChunkType::Function);
+		}
+	}
+
+	{
+		auto P = ParseShader(TEXT("cbuffer Foo {blah} SamplerState S;"));
+
+		if (TestEqual(TEXT("ParseShader: cbuffer and sampler state: num chunks"), P.Chunks.Num(), 2))
+		{
+			TestEqual(TEXT("ParseShader: cbuffer and sampler state: chunk type [0]"), P.Chunks[0].Type, ECodeChunkType::CBuffer);
+			TestEqual(TEXT("ParseShader: cbuffer and sampler state: chunk type [1]"), P.Chunks[1].Type, ECodeChunkType::Variable);
+		}
+	}
+
+	{
+		auto P = ParseShader(TEXT("struct Foo { int a; };"));
+		if (TestEqual(TEXT("ParseShader: struct: num chunks"), P.Chunks.Num(), 1))
+		{
+			TestEqual(TEXT("ParseShader: struct: chunk type"), P.Chunks[0].Type, ECodeChunkType::Struct);
+		}
+	}
+
+	{
+		auto P = ParseShader(TEXT("struct { int a; } Foo = { 123; };"));
+		if (TestEqual(TEXT("ParseShader: anonymous struct with variable and initializer: num chunks"), P.Chunks.Num(), 1))
+		{
+			TestEqual(TEXT("ParseShader: anonymous struct with variable and initializer: chunk type [0]"), P.Chunks[0].Type, ECodeChunkType::Variable);
+		}
+	}
+
+#if 0
+	{
+		// TODO: handle struct forward declarations
+		auto P = ParseShader(TEXT("struct Foo;"));
+	}
+#endif
+
+	{
+		auto P = ParseShader(
+			TEXT("cbuffer MyBuffer : register(b3)")
+			TEXT("{ float4 Element1 : packoffset(c0); float1 Element2 : packoffset(c1); float1 Element3 : packoffset(c1.y); }"));
+		if (TestEqual(TEXT("ParseShader: cbuffer with packoffset: num chunks"), P.Chunks.Num(), 1))
+		{
+			TestEqual(TEXT("ParseShader: cbuffer with packoffset: chunk type"), P.Chunks[0].Type, ECodeChunkType::CBuffer);
+		}
+	}
+
+	{
+		auto P = ParseShader(TEXT("static const struct { float4 Param; } Foo = { FooCB_Param; };"));
+		if (TestEqual(TEXT("ParseShader: static const anonymous struct with variable and initializer: num chunks"), P.Chunks.Num(), 1))
+		{
+			TestEqual(TEXT("ParseShader: static const anonymous struct with variable and initializer: chunk type"), P.Chunks[0].Type, ECodeChunkType::Variable);
+		}
+	}
+
+	{
+		auto P = ParseShader(TEXT("template <typename T> float Fun(T x) { return (float)x; }"));
+		if (TestEqual(TEXT("ParseShader: template function: num chunks"), P.Chunks.Num(), 1))
+		{
+			TestEqual(TEXT("ParseShader: template function: chunk type"), P.Chunks[0].Type, ECodeChunkType::Function);
+		}
+	}
+
+	{
+		auto P = ParseShader(TEXT("enum EFoo { A, B = 123 };"));
+		if (TestEqual(TEXT("ParseShader: enum: num chunks"), P.Chunks.Num(), 1))
+		{
+			TestEqual(TEXT("ParseShader: enum: chunk type"), P.Chunks[0].Type, ECodeChunkType::Enum);
+		}
+	}
+
+	{
+		auto P = ParseShader(TEXT("enum class EFoo { A, B };"));
+		if (TestEqual(TEXT("ParseShader: enum class: num chunks"), P.Chunks.Num(), 1))
+		{
+			TestEqual(TEXT("ParseShader: enum class: chunk type"), P.Chunks[0].Type, ECodeChunkType::Enum);
+		}
+	}
+
+	{
+		auto P = ParseShader(TEXT("#define Foo 123"));
+		if (TestEqual(TEXT("ParseShader: define: num chunks"), P.Chunks.Num(), 1))
+		{
+			TestEqual(TEXT("ParseShader: define: chunk type"), P.Chunks[0].Type, ECodeChunkType::Define);
+		}
+	}
+
+	{
+		auto P = ParseShader(TEXT("#pragma Foo"));
+		if (TestEqual(TEXT("ParseShader: pragma: num chunks"), P.Chunks.Num(), 1))
+		{
+			TestEqual(TEXT("ParseShader: pragma: chunk type"), P.Chunks[0].Type, ECodeChunkType::Pragma);
+		}
+	}
+
+	{
+		auto P = ParseShader(TEXT("ConstantBuffer<Foo> CB : register ( b123, space456);"));
+		if (TestEqual(TEXT("ParseShader: ConstantBuffer<Foo>: num chunks"), P.Chunks.Num(), 1))
+		{
+			TestEqual(TEXT("ParseShader: ConstantBuffer<Foo>: chunk type"), P.Chunks[0].Type, ECodeChunkType::Variable);
+		}
+	}
+
+	int32 NumErrors = ExecutionInfo.GetErrorTotal();
+
+	return NumErrors == 0;
+}
+
+IMPLEMENT_SIMPLE_AUTOMATION_TEST(FShaderMinifierTest, "System.Shaders.ShaderMinifier.Minify", EAutomationTestFlags::EditorContext | EAutomationTestFlags::EngineFilter);
+
+bool FShaderMinifierTest::RunTest(const FString& Parameters)
+{
+	// TODO: run minifier on some tricky examples
+
+	int32 NumErrors = ExecutionInfo.GetErrorTotal();
+	return NumErrors == 0;
+}
+
+#endif // WITH_AUTOMATION_TESTS
diff --git a/Engine/Source/Developer/ShaderCompilerCommon/Public/ShaderMinifier.h b/Engine/Source/Developer/ShaderCompilerCommon/Public/ShaderMinifier.h
new file mode 100644
--- /dev/null
+++ b/Engine/Source/Developer/ShaderCompilerCommon/Public/ShaderMinifier.h
@@ -0,0 +1,50 @@
+// Copyright Epic Games, Inc. All Rights Reserved.
+
+#pragma once
+
+#include "CoreMinimal.h"
+
+#include "Containers/UnrealString.h"
+#include "Misc/EnumClassFlags.h"
+
+namespace UE::ShaderMinifier
+{
+
+enum EMinifyShaderFlags
+{
+	None = 0,
+	OutputReasons = 1 << 1,
+	OutputStats   = 1 << 2,
+};
+ENUM_CLASS_FLAGS(EMinifyShaderFlags)
+
+struct FDiagnosticMessage
+{
+	FString Message;
+	FString File;
+	int32 Offset  = INDEX_NONE;
+	int32 Line    = INDEX_NONE;
+	int32 Column  = INDEX_NONE;
+};
+
+struct FDiagnostics
+{
+	TArray<FDiagnosticMessage> Errors;
+	TArray<FDiagnosticMessage> Warnings;
+};
+
+struct FMinifiedShader
+{
+	FString Code;
+	FDiagnostics Diagnostics;
+
+	bool Success() const
+	{
+		return Diagnostics.Errors.IsEmpty();
+	}
+};
+
+extern SHADERCOMPILERCOMMON_API FMinifiedShader Minify(const FStringView PreprocessedShader, const FStringView EntryPoint, EMinifyShaderFlags Flags = EMinifyShaderFlags::None);
+
+} // UE::ShaderMinifier
+
