ShaderMinifier - Ignore known resource types when looking for potential identifier names. Also add inclusion reason tracking for global variables.

diff --git a/Engine/Source/Developer/ShaderCompilerCommon/Private/ShaderMinifier.cpp b/Engine/Source/Developer/ShaderCompilerCommon/Private/ShaderMinifier.cpp
--- a/Engine/Source/Developer/ShaderCompilerCommon/Private/ShaderMinifier.cpp
+++ b/Engine/Source/Developer/ShaderCompilerCommon/Private/ShaderMinifier.cpp
@@ -898,7 +898,6 @@ static FString MinifyShader(const FParsedShader& Parsed, FStringView SemicolonSe
 
 	{
 		// Some known builtin words to ignore
-
 		ProcessedIdentifiers.Add(TEXT("asfloat"));
 		ProcessedIdentifiers.Add(TEXT("asint"));
 		ProcessedIdentifiers.Add(TEXT("asuint"));
@@ -906,9 +905,11 @@ static FString MinifyShader(const FParsedShader& Parsed, FStringView SemicolonSe
 		ProcessedIdentifiers.Add(TEXT("bool2"));
 		ProcessedIdentifiers.Add(TEXT("bool3"));
 		ProcessedIdentifiers.Add(TEXT("bool4"));
+		ProcessedIdentifiers.Add(TEXT("break"));
 		ProcessedIdentifiers.Add(TEXT("cbuffer"));
 		ProcessedIdentifiers.Add(TEXT("const"));
 		ProcessedIdentifiers.Add(TEXT("else"));
+		ProcessedIdentifiers.Add(TEXT("extern"));
 		ProcessedIdentifiers.Add(TEXT("false"));
 		ProcessedIdentifiers.Add(TEXT("float"));
 		ProcessedIdentifiers.Add(TEXT("float2"));
@@ -918,6 +919,7 @@ static FString MinifyShader(const FParsedShader& Parsed, FStringView SemicolonSe
 		ProcessedIdentifiers.Add(TEXT("float4"));
 		ProcessedIdentifiers.Add(TEXT("float4x4"));
 		ProcessedIdentifiers.Add(TEXT("for"));
+		ProcessedIdentifiers.Add(TEXT("groupshared"));
 		ProcessedIdentifiers.Add(TEXT("if"));
 		ProcessedIdentifiers.Add(TEXT("in"));
 		ProcessedIdentifiers.Add(TEXT("inout"));
@@ -925,11 +927,16 @@ static FString MinifyShader(const FParsedShader& Parsed, FStringView SemicolonSe
 		ProcessedIdentifiers.Add(TEXT("int2"));
 		ProcessedIdentifiers.Add(TEXT("int3"));
 		ProcessedIdentifiers.Add(TEXT("int4"));
+		ProcessedIdentifiers.Add(TEXT("interface"));
 		ProcessedIdentifiers.Add(TEXT("out"));
+		ProcessedIdentifiers.Add(TEXT("packoffset"));
+		ProcessedIdentifiers.Add(TEXT("precise"));
+		ProcessedIdentifiers.Add(TEXT("register"));
 		ProcessedIdentifiers.Add(TEXT("return"));
 		ProcessedIdentifiers.Add(TEXT("static"));
 		ProcessedIdentifiers.Add(TEXT("struct"));
 		ProcessedIdentifiers.Add(TEXT("switch"));
+		ProcessedIdentifiers.Add(TEXT("tbuffer"));
 		ProcessedIdentifiers.Add(TEXT("true"));
 		ProcessedIdentifiers.Add(TEXT("uint"));
 		ProcessedIdentifiers.Add(TEXT("uint2"));
@@ -938,7 +945,72 @@ static FString MinifyShader(const FParsedShader& Parsed, FStringView SemicolonSe
 		ProcessedIdentifiers.Add(TEXT("void"));
 		ProcessedIdentifiers.Add(TEXT("while"));
 
-#if 0 // Some shaders define template versions of some built-in functions, so we can't trivially ignore them
+		// HLSL resource types
+		ProcessedIdentifiers.Add(TEXT("TextureCubeArray"));
+		ProcessedIdentifiers.Add(TEXT("TextureCube"));
+		ProcessedIdentifiers.Add(TEXT("TextureBuffer"));
+		ProcessedIdentifiers.Add(TEXT("Texture3D"));
+		ProcessedIdentifiers.Add(TEXT("Texture2DMSArray"));
+		ProcessedIdentifiers.Add(TEXT("Texture2DMS"));
+		ProcessedIdentifiers.Add(TEXT("Texture2DArray"));
+		ProcessedIdentifiers.Add(TEXT("Texture2D"));
+		ProcessedIdentifiers.Add(TEXT("Texture1DArray"));
+		ProcessedIdentifiers.Add(TEXT("Texture1D"));
+		ProcessedIdentifiers.Add(TEXT("StructuredBuffer"));
+		ProcessedIdentifiers.Add(TEXT("SamplerState"));
+		ProcessedIdentifiers.Add(TEXT("SamplerComparisonState"));
+		ProcessedIdentifiers.Add(TEXT("RWTextureCubeArray"));
+		ProcessedIdentifiers.Add(TEXT("RWTextureCube"));
+		ProcessedIdentifiers.Add(TEXT("RWTexture3D"));
+		ProcessedIdentifiers.Add(TEXT("RWTexture2DMSArray"));
+		ProcessedIdentifiers.Add(TEXT("RWTexture2DMS"));
+		ProcessedIdentifiers.Add(TEXT("RWTexture2DArray"));
+		ProcessedIdentifiers.Add(TEXT("RWTexture2D"));
+		ProcessedIdentifiers.Add(TEXT("RWTexture1DArray"));
+		ProcessedIdentifiers.Add(TEXT("RWTexture1D"));
+		ProcessedIdentifiers.Add(TEXT("RWStructuredBuffer"));
+		ProcessedIdentifiers.Add(TEXT("RWByteAddressBuffer"));
+		ProcessedIdentifiers.Add(TEXT("RWBuffer"));
+		ProcessedIdentifiers.Add(TEXT("RaytracingAccelerationStructure"));
+		ProcessedIdentifiers.Add(TEXT("RasterizerOrderedTexture3D"));
+		ProcessedIdentifiers.Add(TEXT("RasterizerOrderedTexture2DArray"));
+		ProcessedIdentifiers.Add(TEXT("RasterizerOrderedTexture2D"));
+		ProcessedIdentifiers.Add(TEXT("RasterizerOrderedTexture1DArray"));
+		ProcessedIdentifiers.Add(TEXT("RasterizerOrderedTexture1D"));
+		ProcessedIdentifiers.Add(TEXT("RasterizerOrderedStructuredBuffer"));
+		ProcessedIdentifiers.Add(TEXT("RasterizerOrderedByteAddressBuffer"));
+		ProcessedIdentifiers.Add(TEXT("RasterizerOrderedBuffer"));
+		ProcessedIdentifiers.Add(TEXT("FeedbackTexture2DArray"));
+		ProcessedIdentifiers.Add(TEXT("FeedbackTexture2D"));
+		ProcessedIdentifiers.Add(TEXT("ConsumeStructuredBuffer"));
+		ProcessedIdentifiers.Add(TEXT("ConstantBuffer"));
+		ProcessedIdentifiers.Add(TEXT("ByteAddressBuffer"));
+		ProcessedIdentifiers.Add(TEXT("Buffer"));
+		ProcessedIdentifiers.Add(TEXT("AppendStructuredBuffer"));
+
+		// Alternative spelling of some resource types
+		ProcessedIdentifiers.Add(TEXT("AppendRegularBuffer"));
+		ProcessedIdentifiers.Add(TEXT("ByteBuffer"));
+		ProcessedIdentifiers.Add(TEXT("ConsumeRegularBuffer"));
+		ProcessedIdentifiers.Add(TEXT("DataBuffer"));
+		ProcessedIdentifiers.Add(TEXT("MS_Texture2D"));
+		ProcessedIdentifiers.Add(TEXT("MS_Texture2D_Array"));
+		ProcessedIdentifiers.Add(TEXT("RegularBuffer"));
+		ProcessedIdentifiers.Add(TEXT("RW_ByteBuffer"));
+		ProcessedIdentifiers.Add(TEXT("RW_DataBuffer"));
+		ProcessedIdentifiers.Add(TEXT("RW_RegularBuffer"));
+		ProcessedIdentifiers.Add(TEXT("RW_Texture1D"));
+		ProcessedIdentifiers.Add(TEXT("RW_Texture1D_Array"));
+		ProcessedIdentifiers.Add(TEXT("RW_Texture2D"));
+		ProcessedIdentifiers.Add(TEXT("RW_Texture2D_Array"));
+		ProcessedIdentifiers.Add(TEXT("RW_Texture3D"));
+		ProcessedIdentifiers.Add(TEXT("RW_TextureCube"));
+		ProcessedIdentifiers.Add(TEXT("Texture1D_Array"));
+		ProcessedIdentifiers.Add(TEXT("Texture2D_Array"));
+		ProcessedIdentifiers.Add(TEXT("TextureBuffer"));
+		ProcessedIdentifiers.Add(TEXT("TextureCube_Array"));
+
+		// Some shaders define template versions of some built-in functions, so we can't trivially ignore them
 		//ProcessedIdentifiers.Add(TEXT("abs"));
 		//ProcessedIdentifiers.Add(TEXT("any"));
 		//ProcessedIdentifiers.Add(TEXT("clamp"));
@@ -957,7 +1029,6 @@ static FString MinifyShader(const FParsedShader& Parsed, FStringView SemicolonSe
 		//ProcessedIdentifiers.Add(TEXT("sign"));
 		//ProcessedIdentifiers.Add(TEXT("sin"));
 		//ProcessedIdentifiers.Add(TEXT("sqrt"));
-#endif
 	}
 
 	if (PendingChunks.IsEmpty())
@@ -1034,8 +1105,10 @@ static FString MinifyShader(const FParsedShader& Parsed, FStringView SemicolonSe
 					{
 						PendingChunks.Push(Chunk);
 
-						// TODO: handle cbuffers, resources and globals
-						if (Chunk->Type == ECodeChunkType::Function || Chunk->Type == ECodeChunkType::Struct)
+						if (Chunk->Type == ECodeChunkType::Function 
+							|| Chunk->Type == ECodeChunkType::Struct
+							|| Chunk->Type == ECodeChunkType::CBuffer
+							|| Chunk->Type == ECodeChunkType::Variable)
 						{
 							ChunkRequestedBy.FindOrAdd(Chunk) = CurrentChunk;
 						}
@@ -1058,7 +1131,6 @@ static FString MinifyShader(const FParsedShader& Parsed, FStringView SemicolonSe
 			continue;
 		}
 
-		// TODO: print chunk info
 		if (Chunk.Type == ECodeChunkType::Function)
 		{
 			NumFunctions += 1;
@@ -1096,8 +1168,6 @@ static FString MinifyShader(const FParsedShader& Parsed, FStringView SemicolonSe
 	{
 		if (Chunk.Type != ECodeChunkType::Pragma      // Pragmas and defines that remain after preprocessing
 			&& Chunk.Type != ECodeChunkType::Define   // must be preserved as they may control important compiler behaviors.
-			// && Chunk.Type != ECodeChunkType::Variable // Varibles must be currently preserved due to how RootParametersStructure is generated later.
-			// && Chunk.Type != ECodeChunkType::CBuffer  // TODO: analyze cbuffer member usage
 			&& RelevantChunks.Find(&Chunk) == nullptr)
 		{
 			continue;
