Cleaning up deprecated items in "Shader" files for 5.2

diff --git a/Engine/Source/Developer/ShaderCompilerCommon/Public/ShaderCompilerCommon.h b/Engine/Source/Developer/ShaderCompilerCommon/Public/ShaderCompilerCommon.h
--- a/Engine/Source/Developer/ShaderCompilerCommon/Public/ShaderCompilerCommon.h
+++ b/Engine/Source/Developer/ShaderCompilerCommon/Public/ShaderCompilerCommon.h
@@ -224,12 +224,6 @@ extern SHADERCOMPILERCOMMON_API void DumpDebugShaderBinary(const FShaderCompiler
 extern SHADERCOMPILERCOMMON_API void DumpDebugShaderDisassembledSpirv(const FShaderCompilerInput& Input, void* InData, int32 InDataByteSize, const FString& FileExtension);
 extern SHADERCOMPILERCOMMON_API void DumpDebugShaderDisassembledDxil(const FShaderCompilerInput& Input, void* InData, int32 InDataByteSize, const FString& FileExtension);
 
-UE_DEPRECATED(4.26, "SourceLength is no longer needed.")
-inline void DumpDebugUSF(const FShaderCompilerInput& Input, const ANSICHAR* Source, int32 SourceLength, uint32 HlslCCFlags = 0, const TCHAR* OverrideBaseFilename = nullptr)
-{
-	DumpDebugUSF(Input, Source, HlslCCFlags, OverrideBaseFilename);
-}
-
 // calls 'Mali Offline Compiler' to compile the glsl source code and extract the generated instruction count
 extern SHADERCOMPILERCOMMON_API void CompileOfflineMali(const FShaderCompilerInput &Input, FShaderCompilerOutput& ShaderOutput, const ANSICHAR* ShaderSource, const int32 SourceSize, bool bVulkanSpirV, const ANSICHAR* VulkanSpirVEntryPoint = nullptr);
 
diff --git a/Engine/Source/Developer/TargetPlatform/Public/Interfaces/IShaderFormat.h b/Engine/Source/Developer/TargetPlatform/Public/Interfaces/IShaderFormat.h
--- a/Engine/Source/Developer/TargetPlatform/Public/Interfaces/IShaderFormat.h
+++ b/Engine/Source/Developer/TargetPlatform/Public/Interfaces/IShaderFormat.h
@@ -113,13 +113,6 @@ public:
 	 */
 	virtual void NotifyShaderCompiled(const TConstArrayView<uint8>& PlatformDebugData, FName Format) const { }
 
-	UE_DEPRECATED(5.0, "NotifyShaderCompiled should be called")
-	virtual void NotifyShaderCooked(const TConstArrayView<uint8>& PlatformDebugData, FName Format) const
-	{
-		NotifyShaderCompiled(PlatformDebugData, Format);
-	}
-
-
 	/**
 	 * Appends shader key text to the provided key string for use in DDC unique key construction.
 	 * @param KeyString String that will get shader key text appended to.
diff --git a/Engine/Source/Runtime/RenderCore/Private/ShaderResource.cpp b/Engine/Source/Runtime/RenderCore/Private/ShaderResource.cpp
--- a/Engine/Source/Runtime/RenderCore/Private/ShaderResource.cpp
+++ b/Engine/Source/Runtime/RenderCore/Private/ShaderResource.cpp
@@ -477,18 +477,6 @@ void FShaderMapResourceCode::NotifyShadersCompiled(FName FormatName)
 	}
 #endif // WITH_ENGINE
 }
-
-void FShaderMapResourceCode::NotifyShadersCooked(const ITargetPlatform* TargetPlatform)
-{
-#if WITH_ENGINE
-	TArray<FName> ShaderFormatNames;
-	TargetPlatform->GetAllTargetedShaderFormats(ShaderFormatNames);
-	for (FName FormatName : ShaderFormatNames)
-	{
-		NotifyShadersCompiled(FormatName);
-	}
-#endif
-}
 #endif // WITH_EDITORONLY_DATA
 
 FShaderMapResource::FShaderMapResource(EShaderPlatform InPlatform, int32 NumShaders)
diff --git a/Engine/Source/Runtime/RenderCore/Public/Shader.h b/Engine/Source/Runtime/RenderCore/Public/Shader.h
--- a/Engine/Source/Runtime/RenderCore/Public/Shader.h
+++ b/Engine/Source/Runtime/RenderCore/Public/Shader.h
@@ -78,9 +78,6 @@ struct FShaderCompiledShaderInitializerType;
 struct FShaderCompilerOutput;
 using FShaderMapAssetPaths = TSet<FName>; // Copied from ShaderCodeLibrary.h
 
-UE_DEPRECATED(4.26, "FShadereCompiledShaderInitializerType is deprecated. Use FShaderCompiledShaderInitializerType.")
-typedef FShaderCompiledShaderInitializerType FShadereCompiledShaderInitializerType;
-
 /** Define a shader permutation uniquely according to its type, and permutation id.*/
 template<typename MetaShaderType>
 struct TShaderTypePermutation
@@ -491,8 +488,6 @@ public:
 	RENDERCORE_API void Serialize(FArchive& Ar, bool bLoadedByCookedMaterial);
 #if WITH_EDITORONLY_DATA
 	RENDERCORE_API void NotifyShadersCompiled(FName FormatName);
-	UE_DEPRECATED(5.0, "NotifyShadersCompiled should be called")
-	RENDERCORE_API void NotifyShadersCooked(const ITargetPlatform* TargetPlatform);
 #endif // WITH_EDITORONLY_DATA
 
 	RENDERCORE_API uint32 GetSizeBytes() const;
diff --git a/Engine/Source/Runtime/RenderCore/Public/ShaderCompilerCore.h b/Engine/Source/Runtime/RenderCore/Public/ShaderCompilerCore.h
--- a/Engine/Source/Runtime/RenderCore/Public/ShaderCompilerCore.h
+++ b/Engine/Source/Runtime/RenderCore/Public/ShaderCompilerCore.h
@@ -45,11 +45,6 @@ extern RENDERCORE_API bool ShouldEnableExtraShaderData(FName ShaderFormat);
 
 extern RENDERCORE_API bool ShouldOptimizeShaders(FName ShaderFormat);
 
-UE_DEPRECATED(5.0, "ShouldGenerateShaderSymbols should be called to determine if symbols (debug data) should be generated")
-bool ShouldKeepShaderDebugInfo(EShaderPlatform Platform);
-UE_DEPRECATED(5.0, "ShouldWriteShaderSymbols should be called to determine if symbols (debug data) should be written")
-bool ShouldExportShaderDebugInfo(EShaderPlatform Platform);
-
 /** Returns true is shader compiling is allowed */
 extern RENDERCORE_API bool AllowShaderCompiling();
 
@@ -70,7 +65,6 @@ enum ECompilerFlags
 	CFLAG_ForceOptimization,
 	// Shader should generate symbols for debugging.
 	CFLAG_GenerateSymbols,
-	CFLAG_KeepDebugInfo UE_DEPRECATED(5.0, "CFLAG_GenerateSymbols should be used to signal if debug data needs to be generated") = CFLAG_GenerateSymbols,
 	// Shader should insert debug/name info at the risk of generating non-deterministic libraries
 	CFLAG_ExtraShaderData,
 	// Allows the (external) symbols to be specific to each shader rather than trying to deduplicate.
