Fixing parsing resource types like "Texture3D < uint4 >" where the space between the type and the "<" prevented the type from being found.

diff --git a/Engine/Source/Developer/ShaderCompilerCommon/Private/ShaderCompilerCommon.cpp b/Engine/Source/Developer/ShaderCompilerCommon/Private/ShaderCompilerCommon.cpp
--- a/Engine/Source/Developer/ShaderCompilerCommon/Private/ShaderCompilerCommon.cpp
+++ b/Engine/Source/Developer/ShaderCompilerCommon/Private/ShaderCompilerCommon.cpp
@@ -527,8 +527,9 @@ EShaderParameterType UE::ShaderCompilerCommon::ParseParameterType(
 	FStringView UntemplatedType = InType;
 	if (int32 Index = InType.Find(TEXT("<")); Index != INDEX_NONE)
 	{
+		// Remove the template argument but don't forget to clean up the type name
 		const int32 NumChars = InType.Len() - Index;
-		UntemplatedType = InType.LeftChop(NumChars);
+		UntemplatedType = InType.LeftChop(NumChars).TrimEnd();
 	}
 
 	if (AllSRVTypes.Contains(UntemplatedType) || InExtraSRVTypes.Contains(UntemplatedType))
