Bindless Resource shader support
- Adding better support for resources that have multiple template arguments (Texture2DMS<float4,2>, etc).
- Requires typedef'ing all Resource types to avoid their extra comments being given to our AUTO_BINDLESS macros
- Adding support for globallycoherent
- Reworked BindlessConversionType to use its own enum.

diff --git a/Engine/Source/Developer/ShaderCompilerCommon/Private/ShaderParameterParser.cpp b/Engine/Source/Developer/ShaderCompilerCommon/Private/ShaderParameterParser.cpp
--- a/Engine/Source/Developer/ShaderCompilerCommon/Private/ShaderParameterParser.cpp
+++ b/Engine/Source/Developer/ShaderCompilerCommon/Private/ShaderParameterParser.cpp
@@ -213,6 +213,7 @@ bool FShaderParameterParser::ParseParameters(
 		int32 ArrayStartPos = -1;
 		int32 ArrayEndPos = -1;
 		int32 ScopeIndent = 0;
+		bool bGloballyCoherent = false;
 
 		EState State = EState::Scanning;
 		bool bGoToNextLine = false;
@@ -226,6 +227,7 @@ bool FShaderParameterParser::ParseParameters(
 			NameEndPos = -1;
 			ArrayStartPos = -1;
 			ArrayEndPos = -1;
+			bGloballyCoherent = false;
 			State = EState::Scanning;
 		};
 
@@ -264,7 +266,8 @@ bool FShaderParameterParser::ParseParameters(
 				EShaderParameterType ParsedParameterType = UE::ShaderCompilerCommon::ParseAndRemoveBindlessParameterPrefix(Name);
 				const bool bBindlessIndex = (ParsedParameterType == EShaderParameterType::BindlessResourceIndex || ParsedParameterType == EShaderParameterType::BindlessSamplerIndex);
 
-				EShaderParameterType BindlessConversionType = EShaderParameterType::Num;
+				EBindlessConversionType BindlessConversionType = EBindlessConversionType::None;
+				EShaderParameterType ConstantBufferType = ParsedParameterType;
 
 				if ((bBindlessResources || bBindlessSamplers) && ParsedParameterType == EShaderParameterType::LooseData)
 				{
@@ -272,17 +275,20 @@ bool FShaderParameterParser::ParseParameters(
 
 					if (bBindlessResources && (ParsedParameterType == EShaderParameterType::SRV || ParsedParameterType == EShaderParameterType::UAV))
 					{
-						BindlessConversionType = EShaderParameterType::BindlessResourceIndex;
+						BindlessConversionType = EBindlessConversionType::Resource;
+						ConstantBufferType = EShaderParameterType::BindlessResourceIndex;
 					}
 					else if (bBindlessSamplers && ParsedParameterType == EShaderParameterType::Sampler)
 					{
-						BindlessConversionType = EShaderParameterType::BindlessSamplerIndex;
+						BindlessConversionType = EBindlessConversionType::Sampler;
+						ConstantBufferType = EShaderParameterType::BindlessSamplerIndex;
 					}
 
-					if (BindlessConversionType != EShaderParameterType::Num && Leftovers.Contains(TEXT("register")))
+					if (BindlessConversionType != EBindlessConversionType::None && Leftovers.Contains(TEXT("register")))
 					{
 						// avoid rewriting hardcoded register assignments
-						BindlessConversionType = EShaderParameterType::Num;
+						BindlessConversionType = EBindlessConversionType::None;
+						ConstantBufferType = ParsedParameterType;
 					}
 				}
 
@@ -314,11 +320,11 @@ bool FShaderParameterParser::ParseParameters(
 								BaseType == UBMT_UINT32 ||
 								BaseType == UBMT_FLOAT32 ||
 								bBindlessIndex ||
-								(BindlessConversionType != EShaderParameterType::Num);
+								(BindlessConversionType != EBindlessConversionType::None);
 
 							if (bMoveToRootConstantBuffer)
 							{
-								ConstantBufferParameterType = BindlessConversionType != EShaderParameterType::Num ? BindlessConversionType : ParsedParameterType;
+								ConstantBufferParameterType = ConstantBufferType;
 							}
 						}
 					}
@@ -339,6 +345,7 @@ bool FShaderParameterParser::ParseParameters(
 					ParsedParameter.ParsedCharOffsetEnd = Cursor;
 					ParsedParameter.BindlessConversionType = BindlessConversionType;
 					ParsedParameter.ConstantBufferParameterType = ConstantBufferParameterType;
+					ParsedParameter.bGloballyCoherent = bGloballyCoherent;
 
 					if (ArrayStartPos != -1 && ArrayEndPos != -1)
 					{
@@ -405,13 +412,14 @@ bool FShaderParameterParser::ParseParameters(
 				{
 					static const TCHAR* KeywordTable[] =
 					{
+						TEXT("const"),
+						TEXT("globallycoherent")
 						TEXT("enum"),
 						TEXT("class"),
-						TEXT("const"),
 						TEXT("struct"),
 						TEXT("static"),
 					};
-					static int32 KeywordTableSize[] = { 4, 5, 5, 6, 6 };
+					static int32 KeywordTableSize[] = { 5, 16, 4, 5, 6, 6 };
 
 					int32 RecognisedKeywordId = -1;
 					for (int32 KeywordId = 0; KeywordId < UE_ARRAY_COUNT(KeywordTable); KeywordId++)
@@ -438,7 +446,7 @@ bool FShaderParameterParser::ParseParameters(
 						State = EState::ParsingPotentialType;
 						TypeStartPos = Cursor;
 					}
-					else if (RecognisedKeywordId == 2)
+					else if (RecognisedKeywordId == 0)
 					{
 						// Ignore the const keywords, but still parse given it might still be a shader parameter.
 						if (TypeQualifierStartPos == -1)
@@ -448,6 +456,17 @@ bool FShaderParameterParser::ParseParameters(
 						}
 						Cursor += KeywordTableSize[RecognisedKeywordId];
 					}
+					else if (RecognisedKeywordId == 1)
+					{
+						// Mark that we got the globallycoherent keyword and keep moving to the next set of qualifiers
+						bGloballyCoherent = true;
+						if (TypeQualifierStartPos == -1)
+						{
+							// If the parameter is erased, we also have to erase *all* qualifiers, e.g. "const int Foo" or "const const int Foo".
+							TypeQualifierStartPos = Cursor;
+						}
+						Cursor += KeywordTableSize[RecognisedKeywordId];
+					}
 					else
 					{
 						// Purposefully ignore enum, class, struct, static
@@ -487,7 +506,7 @@ bool FShaderParameterParser::ParseParameters(
 				}
 				else if (Char == '<')
 				{
-					// Found what looks like the begining of template argument that is legal on resource types for Instance Texture2D< float >
+					// Found what looks like the beginning of template argument that is legal on resource types for Instance Texture2D< float >
 					State = EState::ParsingPotentialTypeTemplateArguments;
 				}
 				else if (bIsWhiteSpace)
@@ -688,7 +707,8 @@ void FShaderParameterParser::RemoveMovingParametersFromSource(FString& Preproces
 		// If this parameter is going to be in the root constant buffer
 		if (ParsedParameter.ConstantBufferParameterType != EShaderParameterType::Num &&
 			// but it's not being converted to bindless
-			ParsedParameter.BindlessConversionType == EShaderParameterType::Num)
+			ParsedParameter.BindlessConversionType == EBindlessConversionType::None &&
+			ParsedParameter.ParsedCharOffsetStart != INDEX_NONE)
 		{
 			// then erase this shader parameter conserving the same line numbers.
 			for (int32 j = ParsedParameter.ParsedCharOffsetStart; j <= ParsedParameter.ParsedCharOffsetEnd; j++)
@@ -726,18 +746,28 @@ void FShaderParameterParser::ApplyBindlessModifications(FString& PreprocessedSha
 				continue;
 			}
 
-			if (ParsedParameter.BindlessConversionType != EShaderParameterType::Num)
+			if (ParsedParameter.BindlessConversionType != EBindlessConversionType::None)
 			{
-				const bool bIsSampler = (ParsedParameter.BindlessConversionType == EShaderParameterType::BindlessSamplerIndex);
+				const bool bIsSampler = (ParsedParameter.BindlessConversionType == EBindlessConversionType::Sampler);
 
 				const TCHAR* IndexDefine = bIsSampler ? TEXT("AUTO_BINDLESS_SAMPLER_INDEX") : TEXT("AUTO_BINDLESS_RESOURCE_INDEX");
 				const TCHAR* VariableDefine = bIsSampler ? TEXT("AUTO_BINDLESS_SAMPLER_VARIABLE") : TEXT("AUTO_BINDLESS_RESOURCE_VARIABLE");
+				const TCHAR* StorageClass = ParsedParameter.bGloballyCoherent ? TEXT("globallycoherent ") : TEXT("");
 
 				const FStringView Name = ShaderBindingName;
 				const FStringView Type = ParsedParameter.ParsedType;
 
+				FString RewriteType(Type);
+
+				FString TypedefText = TEXT("");
+				if (ParsedParameter.BindlessConversionType == EBindlessConversionType::Resource)
+				{
+					RewriteType = FString::Printf(TEXT("SafeType%.*s"), Name.Len(), Name.GetData());
+					TypedefText = FString::Printf(TEXT("typedef %.*s %s;"), Type.Len(), Type.GetData(), *RewriteType);
+				}
+
 				const FString IndexText = FString::Printf(TEXT("%s(%.*s);"), IndexDefine, Name.Len(), Name.GetData());
-				const FString VariableText = FString::Printf(TEXT("%s(%.*s, %.*s);"), VariableDefine, Type.Len(), Type.GetData(), Name.Len(), Name.GetData());
+				const FString VariableText = FString::Printf(TEXT("%s(%s%s, %.*s);"), VariableDefine, StorageClass, *RewriteType, Name.Len(), Name.GetData());
 
 				FShaderCodeModifications Modif;
 				Modif.CharOffsetStart = ParsedParameter.ParsedCharOffsetStart;
@@ -748,6 +778,7 @@ void FShaderParameterParser::ApplyBindlessModifications(FString& PreprocessedSha
 				{
 					Modif.Replace += IndexText;
 				}
+				Modif.Replace += TypedefText;
 				Modif.Replace += VariableText;
 
 				Modifications.Add(Modif);
diff --git a/Engine/Source/Developer/ShaderCompilerCommon/Public/ShaderParameterParser.h b/Engine/Source/Developer/ShaderCompilerCommon/Public/ShaderParameterParser.h
--- a/Engine/Source/Developer/ShaderCompilerCommon/Public/ShaderParameterParser.h
+++ b/Engine/Source/Developer/ShaderCompilerCommon/Public/ShaderParameterParser.h
@@ -10,6 +10,13 @@ struct FShaderCompilerOutput;
 
 enum class EShaderParameterType : uint8;
 
+enum class EBindlessConversionType : uint8
+{
+	None,
+	Resource,
+	Sampler
+};
+
 /** Validates and moves all the shader loose data parameter defined in the root scope of the shader into the root uniform buffer. */
 class SHADERCOMPILERCOMMON_API FShaderParameterParser
 {
@@ -44,12 +51,14 @@ public:
 		int32 ParsedLineOffset = 0;
 
 		/** Character position of the start and end of the parameter decelaration in FParsedShaderParameter::OriginalParsedShader */
-		int32 ParsedCharOffsetStart = 0;
-		int32 ParsedCharOffsetEnd = 0;
+		int32 ParsedCharOffsetStart = INDEX_NONE;
+		int32 ParsedCharOffsetEnd = INDEX_NONE;
 
-		EShaderParameterType BindlessConversionType{};
+		EBindlessConversionType BindlessConversionType{};
 		EShaderParameterType ConstantBufferParameterType{};
 
+		bool bGloballyCoherent = false;
+
 		friend class FShaderParameterParser;
 	};
 
