Make the parsed resource/sampler name be the authoritative casing of the variable names. The rest of the shader system is case insensitive but the shader code that uses the variables isn't, so we need to preserve the case we found in the shader.

diff --git a/Engine/Source/Developer/ShaderCompilerCommon/Private/ShaderParameterParser.cpp b/Engine/Source/Developer/ShaderCompilerCommon/Private/ShaderParameterParser.cpp
--- a/Engine/Source/Developer/ShaderCompilerCommon/Private/ShaderParameterParser.cpp
+++ b/Engine/Source/Developer/ShaderCompilerCommon/Private/ShaderParameterParser.cpp
@@ -338,6 +338,7 @@ bool FShaderParameterParser::ParseParameters(
 				// Update 
 				if (bUpdateParsedParameters)
 				{
+					ParsedParameter.ParsedName = Name;
 					ParsedParameter.ParsedType = Type;
 					ParsedParameter.ParsedPragmaLineoffset = CurrentPragamLineoffset;
 					ParsedParameter.ParsedLineOffset = CurrentLineoffset;
@@ -738,7 +739,6 @@ void FShaderParameterParser::ApplyBindlessModifications(FString& PreprocessedSha
 
 		for (TPair<FString, FParsedShaderParameter>& Itr : ParsedParameters)
 		{
-			const FStringView ShaderBindingName = Itr.Key;
 			FParsedShaderParameter& ParsedParameter = Itr.Value;
 
 			if (!ParsedParameter.IsFound())
@@ -754,7 +754,7 @@ void FShaderParameterParser::ApplyBindlessModifications(FString& PreprocessedSha
 				const TCHAR* VariableDefine = bIsSampler ? TEXT("AUTO_BINDLESS_SAMPLER_VARIABLE") : TEXT("AUTO_BINDLESS_RESOURCE_VARIABLE");
 				const TCHAR* StorageClass = ParsedParameter.bGloballyCoherent ? TEXT("globallycoherent ") : TEXT("");
 
-				const FStringView Name = ShaderBindingName;
+				const FStringView Name = ParsedParameter.ParsedName;
 				const FStringView Type = ParsedParameter.ParsedType;
 
 				FString RewriteType(Type);
@@ -874,16 +874,16 @@ bool FShaderParameterParser::MoveShaderParametersToRootConstantBuffer(
 				if (ParsedParameter->ConstantBufferParameterType == EShaderParameterType::BindlessResourceIndex)
 				{
 					RootCBufferContent.Append(FString::Printf(
-						TEXT("uint BindlessResource_%s : packoffset(c%d%s);\r\n"),
-						ShaderBindingName,
+						TEXT("uint BindlessResource_%.*s : packoffset(c%d%s);\r\n"),
+						SVARG(ParsedParameter->ParsedName),
 						ConstantRegister,
 						ConstantSwizzle));
 				}
 				else if (ParsedParameter->ConstantBufferParameterType == EShaderParameterType::BindlessSamplerIndex)
 				{
 					RootCBufferContent.Append(FString::Printf(
-						TEXT("uint BindlessSampler_%s : packoffset(c%d%s);\r\n"),
-						ShaderBindingName,
+						TEXT("uint BindlessSampler_%.*s : packoffset(c%d%s);\r\n"),
+						SVARG(ParsedParameter->ParsedName),
 						ConstantRegister,
 						ConstantSwizzle));
 				}
diff --git a/Engine/Source/Developer/ShaderCompilerCommon/Public/ShaderParameterParser.h b/Engine/Source/Developer/ShaderCompilerCommon/Public/ShaderParameterParser.h
--- a/Engine/Source/Developer/ShaderCompilerCommon/Public/ShaderParameterParser.h
+++ b/Engine/Source/Developer/ShaderCompilerCommon/Public/ShaderParameterParser.h
@@ -28,6 +28,7 @@ public:
 		const FShaderParametersMetadata::FMember* Member = nullptr;
 
 		/** Information found about the member when parsing the preprocessed code. */
+		FStringView ParsedName; /** View into FShaderParameterParser::OriginalParsedShader */
 		FStringView ParsedType; /** View into FShaderParameterParser::OriginalParsedShader */
 		FStringView ParsedArraySize; /** View into FShaderParameterParser::OriginalParsedShader */
 
