ShaderMinifier - Track and output line directives

diff --git a/Engine/Source/Developer/ShaderCompilerCommon/Private/ShaderMinifier.cpp b/Engine/Source/Developer/ShaderCompilerCommon/Private/ShaderMinifier.cpp
--- a/Engine/Source/Developer/ShaderCompilerCommon/Private/ShaderMinifier.cpp
+++ b/Engine/Source/Developer/ShaderCompilerCommon/Private/ShaderMinifier.cpp
@@ -8,13 +8,13 @@
 #include "Logging/LogMacros.h"
 #include "Misc/AutomationTest.h"
 #include "String/Find.h"
+#include "Algo/BinarySearch.h"
 
 DEFINE_LOG_CATEGORY_STATIC(LogShaderMinifier, Log, All);
 
 // TODO:
 // - track namespaces
 // - preserve multi-line #define
-// - track and re-emit correct #line
 
 namespace UE::ShaderMinifier
 {
@@ -90,6 +90,20 @@ static FStringView SkipUntilNonIdentifierCharacter(FStringView Source) {
 	return FStringView(SourceData + Cursor, Len - Cursor);
 }
 
+static FStringView SkipUntilNonNumber(FStringView Source) {
+	int32 Len = Source.Len();
+	int32 Cursor = 0;
+	const TCHAR* SourceData = Source.GetData();
+	while (Cursor < Len)
+	{
+		if (!IsNumber(SourceData[Cursor]))
+		{
+			break;
+		}
+		++Cursor;
+	}
+	return FStringView(SourceData + Cursor, Len - Cursor);
+}
 
 static FStringView SkipSpace(FStringView Source)
 {
@@ -252,6 +266,11 @@ struct FCodeChunk
 	ECodeChunkType Type = ECodeChunkType::Unknown;
 	TArray<FCodeBlock> Blocks;
 
+	// Indicates whether the code for this chunk can be used as-is.
+	// One example where we have to do custom code emission is when a named struct and a variable are declared in one chunk.
+	// The struct type may be referenced, but the variable may be removed. In this case we have to emit the type declaration only.
+	bool bVerbatim = true;
+
 	FStringView FindFirstBlockByType(EBlockType InType) const
 	{
 		for (const FCodeBlock& Block : Blocks)
@@ -263,10 +282,29 @@ struct FCodeChunk
 		}
 		return {};
 	}
+
+	// String view covering the entire code chunk
+	FStringView GetCode() const
+	{
+		if (Blocks.IsEmpty())
+		{
+			return {};
+		}
+		else
+		{
+			const FCodeBlock& FirstBlock = Blocks[0];
+			const FCodeBlock& LastBlock = Blocks[Blocks.Num()-1];
+			const TCHAR* Begin = FirstBlock.Code.GetData();
+			const TCHAR* End = LastBlock.Code.GetData() + LastBlock.Code.Len();
+			return FStringView(Begin, int32(End-Begin));
+		}
+	}
 };
 
 struct FParsedShader
 {
+	FStringView Source;
+
 	TArray<FCodeChunk> Chunks;
 };
 
@@ -274,6 +312,8 @@ static FParsedShader ParseShader(FStringView InSource, FDiagnostics& Output)
 {
 	FParsedShader Result;
 
+	Result.Source = InSource;
+
 	FStringView		   Source = InSource;
 	TArray<FCodeBlock> PendingBlocks;
 	TArray<FCodeChunk> Chunks;
@@ -400,6 +440,7 @@ static FParsedShader ParseShader(FStringView InSource, FDiagnostics& Output)
 
 				FCodeChunk StructChunk;
 				StructChunk.Type = ECodeChunkType::Struct;
+				StructChunk.bVerbatim = false;
 
 				for (int32 i = int32(StructBlockIndex); i < NameBlockIndex; ++i)
 				{
@@ -408,6 +449,7 @@ static FParsedShader ParseShader(FStringView InSource, FDiagnostics& Output)
 
 				FCodeChunk VarChunk;
 				VarChunk.Type = ECodeChunkType::Variable;
+				VarChunk.bVerbatim = false;
 
 				for (int32 i = 0; i < PendingBlocks.Num(); ++i)
 				{
@@ -821,31 +863,39 @@ static void OutputChunk(const FCodeChunk& Chunk, FStringBuilderBase& OutputStrea
 		return;
 	}
 
-	int32 Index = 0;
-	for (const FCodeBlock& Block : Chunk.Blocks)
+	if (Chunk.bVerbatim)
 	{
-		if (Index != 0)
+		// Fast path to output entire code block verbatim, preserving any new lines and whitespace		
+		OutputStream << Chunk.GetCode();
+	}
+	else
+	{
+		int32 Index = 0;
+		for (const FCodeBlock& Block : Chunk.Blocks)
 		{
-			OutputStream << ' ';
-		}
+			if (Index != 0)
+			{
+				OutputStream << ' ';
+			}
 
-		if (Block.Type == EBlockType::Expression)
-		{
-			OutputStream << "= ";
-		}
-		else if (Block.Type == EBlockType::Body)
-		{
-			OutputStream << "\n";
-		}
+			if (Block.Type == EBlockType::Expression)
+			{
+				OutputStream << "= ";
+			}
+			else if (Block.Type == EBlockType::Body)
+			{
+				OutputStream << "\n";
+			}
 
-		if (Block.Type == EBlockType::Binding || Block.Type == EBlockType::Base)
-		{
-			OutputStream << ": ";
-		}
+			if (Block.Type == EBlockType::Binding || Block.Type == EBlockType::Base)
+			{
+				OutputStream << ": ";
+			}
 
-		OutputStream << Block.Code;
+			OutputStream << Block.Code;
 
-		++Index;
+			++Index;
+		}
 	}
 
 	if (Chunk.Type != ECodeChunkType::Function
@@ -871,6 +921,104 @@ struct FCasedStringViewKeyFuncs : public DefaultKeyFuncs<FStringView>
 	}
 };
 
+static void BuildLineBreakMap(FStringView Source, TArray<int32>& OutLineBreakMap, TArray<FStringView>& OutLineDirectives)
+{
+	OutLineBreakMap.Empty();
+	OutLineDirectives.Empty();
+
+	OutLineBreakMap.Add(0); // Lines numbers are 1-based, so add a dummy element to make LowerBound later return the line number directly
+
+	const int32 Len = Source.Len();
+	const TCHAR* Chars = Source.GetData(); // avoid bounds check overhead in [] operator
+
+	for (int32 Index = 0; Index < Len; ++Index)
+	{
+		if (Chars[Index] == TCHAR('\n'))
+		{
+			OutLineBreakMap.Add(Index);
+		}
+		else if (Chars[Index] == TCHAR('#'))
+		{
+			// In a general case, directives may be inside comments or inactive blocks.
+			// However we expect input source to be fully preprocessed and comments to be removed.
+
+			FStringView PossibleDirective = Source.Mid(Index);
+			if (PossibleDirective.StartsWith(TEXT("#line")))
+			{
+				FStringView Remainder = SkipUntilNextLine(PossibleDirective);
+				FStringView LineDirective = SubStrView(PossibleDirective, 0, PossibleDirective.Len() - Remainder.Len());
+				OutLineDirectives.Add(LineDirective);
+			}
+		}
+	}
+}
+
+
+static int32 FindLineDirective(const TArray<FStringView>& LineDirectives, const TCHAR* Ptr)
+{
+	int32 FoundIndex = Algo::UpperBoundBy(LineDirectives, Ptr, [](FStringView Item)
+	{
+		return Item.GetData();
+	});
+
+	if (FoundIndex < 1 || FoundIndex > LineDirectives.Num())
+	{
+		return INDEX_NONE;
+	}
+
+	// UpperBound returns element that's greater than predicate, but we need the closest preceeding line directive.
+	return FoundIndex - 1;
+}
+
+static int32 FindLineNumber(FStringView Source, const TArray<int32>& LineBreakMap, const TCHAR* Ptr)
+{
+	if (Ptr < Source.GetData() || Ptr >= Source.GetData() + Source.Len())
+	{
+		return INDEX_NONE;
+	}
+
+	const int32 Index = int32(Ptr - Source.GetData());
+
+	int32 FoundLineNumber = Algo::UpperBound(LineBreakMap, Index);
+
+	return FoundLineNumber;
+}
+
+static bool ParseLineDirective(FStringView Input, int32& OutLineNumber, FStringView& OutFileName)
+{
+	if (!Input.StartsWith(TEXT("#line")))
+	{
+		return false;
+	}
+
+	Input = Input.Mid(5); // skip `#line` itself
+	Input = SkipSpace(Input);
+
+	if (Input.IsEmpty() || !IsNumber(Input[0]))
+	{
+		return false;
+	}
+
+	OutLineNumber = FCString::Atoi(Input.GetData());
+
+	int32 FileNameBeginIndex = INDEX_NONE;
+	if (Input.FindChar(TCHAR('"'), FileNameBeginIndex))
+	{
+		int32 FileNameEndIndex = INDEX_NONE;
+		Input.MidInline(FileNameBeginIndex + 1);
+		if (Input.FindChar(TCHAR('"'), FileNameEndIndex))
+		{
+			OutFileName = Input.Mid(0, FileNameEndIndex);
+		}
+		else
+		{
+			return false;
+		}
+	}
+
+	return true;
+}
+
 static FString MinifyShader(const FParsedShader& Parsed, FStringView SemicolonSeparatedEntryPoints, EMinifyShaderFlags Flags, FDiagnostics& Diagnostics)
 {
 	FStringBuilderBase OutputStream;
@@ -1164,6 +1312,13 @@ static FString MinifyShader(const FParsedShader& Parsed, FStringView SemicolonSe
 		OutputStream << "\n";
 	}
 
+	TArray<int32> LineBreakMap;
+	TArray<FStringView> LineDirectives;
+	if (EnumHasAnyFlags(Flags, EMinifyShaderFlags::OutputLines))
+	{
+		BuildLineBreakMap(Parsed.Source, LineBreakMap, LineDirectives);
+	}
+
 	for (const FCodeChunk& Chunk : Parsed.Chunks)
 	{
 		if (Chunk.Type != ECodeChunkType::Pragma      // Pragmas and defines that remain after preprocessing
@@ -1182,7 +1337,44 @@ static FString MinifyShader(const FParsedShader& Parsed, FStringView SemicolonSe
 				FStringView  RequestedByName  = RequestedByChunk->FindFirstBlockByType(EBlockType::Name);
 				if (!RequestedByName.IsEmpty())
 				{
-					OutputStream << "// REASON: " << RequestedByName << "\n";
+					OutputStream << TEXT("// REASON: ") << RequestedByName << TEXT("\n");
+				}
+			}
+		}
+
+		if (EnumHasAnyFlags(Flags, EMinifyShaderFlags::OutputLines))
+		{
+			const FStringView ChunkCode = Chunk.Blocks[0].Code;
+			int32 LineDirectiveIndex = FindLineDirective(LineDirectives, ChunkCode.GetData());
+			int32 ChunkLine = FindLineNumber(Parsed.Source, LineBreakMap, ChunkCode.GetData());
+
+			if (ChunkLine != INDEX_NONE && LineDirectiveIndex == INDEX_NONE)
+			{
+				// There was no valid line directive for this chunk, but we do know the line in the input source, so just emit that.
+
+				OutputStream << TEXT("#line ") << ChunkLine << TEXT("\n");
+			}
+			else if (ChunkLine != INDEX_NONE && LineDirectiveIndex != INDEX_NONE)
+			{
+				// We have a valid line directive and line number in the input source.
+				// Some of the input source code may have been removed, so we need to adjust 
+				// the line number before emitting the line directive.
+
+				FStringView LineDirective = LineDirectives[LineDirectiveIndex];
+				int32 LineDirectiveLine = FindLineNumber(Parsed.Source, LineBreakMap, LineDirective.GetData());
+				int32 ParsedLineNumber = INDEX_NONE;
+				FStringView ParsedFileName;
+				if (LineDirectiveLine != INDEX_NONE && ParseLineDirective(LineDirective, ParsedLineNumber, ParsedFileName))
+				{
+					int32 OffsetFromLineDirective = ChunkLine - (LineDirectiveLine + 1); // Line directive identifies the *next* line, hence +1 when computing the offset
+					int32 PatchedLineNumber = ParsedLineNumber + OffsetFromLineDirective;
+
+					OutputStream << TEXT("#line ") << PatchedLineNumber;
+					if (!ParsedFileName.IsEmpty())
+					{
+						OutputStream << TEXT(" \"") << ParsedFileName << TEXT("\"");
+					}
+					OutputStream << TEXT("\n");
 				}
 			}
 		}
diff --git a/Engine/Source/Developer/ShaderCompilerCommon/Public/ShaderMinifier.h b/Engine/Source/Developer/ShaderCompilerCommon/Public/ShaderMinifier.h
--- a/Engine/Source/Developer/ShaderCompilerCommon/Public/ShaderMinifier.h
+++ b/Engine/Source/Developer/ShaderCompilerCommon/Public/ShaderMinifier.h
@@ -15,6 +15,7 @@ enum EMinifyShaderFlags
 	None = 0,
 	OutputReasons = 1 << 1,
 	OutputStats   = 1 << 2,
+	OutputLines   = 1 << 3,
 };
 ENUM_CLASS_FLAGS(EMinifyShaderFlags)
 
diff --git a/Engine/Source/Developer/Windows/ShaderFormatD3D/Private/D3DShaderCompiler.cpp b/Engine/Source/Developer/Windows/ShaderFormatD3D/Private/D3DShaderCompiler.cpp
--- a/Engine/Source/Developer/Windows/ShaderFormatD3D/Private/D3DShaderCompiler.cpp
+++ b/Engine/Source/Developer/Windows/ShaderFormatD3D/Private/D3DShaderCompiler.cpp
@@ -1346,7 +1346,9 @@ void CompileD3DShader(const FShaderCompilerInput& Input, FShaderCompilerOutput&
 		}
 
 		UE::ShaderMinifier::FMinifiedShader Minified  = UE::ShaderMinifier::Minify(PreprocessedShaderSource, EntryMain, 
-			UE::ShaderMinifier::EMinifyShaderFlags::OutputReasons | UE::ShaderMinifier::EMinifyShaderFlags::OutputStats);
+			UE::ShaderMinifier::EMinifyShaderFlags::OutputReasons
+			| UE::ShaderMinifier::EMinifyShaderFlags::OutputStats
+			| UE::ShaderMinifier::EMinifyShaderFlags::OutputLines);
 
 		if (Minified.Success())
 		{
