ShaderParameterParser cleanup
- Making ConstantBufferType a constructor only argument. This removes the duplicated logic of calling ShouldUseStableConstantBuffer.
- Moving ShouldUseStableConstantBuffer to ShaderCompilerCommon.
- Made UE::ShaderCompilerCommon::ParseParameterType array arguments const

diff --git a/Engine/Source/Developer/Apple/MetalShaderFormat/Private/MetalShaderCompiler.cpp b/Engine/Source/Developer/Apple/MetalShaderFormat/Private/MetalShaderCompiler.cpp
--- a/Engine/Source/Developer/Apple/MetalShaderFormat/Private/MetalShaderCompiler.cpp
+++ b/Engine/Source/Developer/Apple/MetalShaderFormat/Private/MetalShaderCompiler.cpp
@@ -1066,7 +1066,7 @@ void CompileShader_Metal(const FShaderCompilerInput& _Input,FShaderCompilerOutpu
 	}
 
 	FShaderParameterParser ShaderParameterParser;
-	if (!ShaderParameterParser.ParseAndModify(Input, Output, PreprocessedShader, /* ConstantBufferType = */ nullptr))
+	if (!ShaderParameterParser.ParseAndModify(Input, Output, PreprocessedShader))
 	{
 		// The FShaderParameterParser will add any relevant errors.
 		return;
diff --git a/Engine/Source/Developer/ShaderCompilerCommon/Private/ShaderCompilerCommon.cpp b/Engine/Source/Developer/ShaderCompilerCommon/Private/ShaderCompilerCommon.cpp
--- a/Engine/Source/Developer/ShaderCompilerCommon/Private/ShaderCompilerCommon.cpp
+++ b/Engine/Source/Developer/ShaderCompilerCommon/Private/ShaderCompilerCommon.cpp
@@ -471,6 +471,18 @@ TCHAR* FindNextUniformBufferReference(TCHAR* SearchPtr, const TCHAR* SearchStrin
 	return nullptr;
 }
 
+bool UE::ShaderCompilerCommon::ShouldUseStableConstantBuffer(const FShaderCompilerInput& Input)
+{
+	// stable constant buffer is for the FShaderParameterBindings::BindForLegacyShaderParameters() code path.
+	// Ray tracing shaders use FShaderParameterBindings::BindForRootShaderParameters instead.
+	if (Input.IsRayTracingShader())
+	{
+		return false;
+	}
+
+	return Input.RootParametersStructure != nullptr;
+}
+
 static const TCHAR* const s_AllSRVTypes[] =
 {
 	TEXT("Texture1D"),
@@ -512,8 +524,8 @@ static const TCHAR* const s_AllSamplerTypes[] =
 
 EShaderParameterType UE::ShaderCompilerCommon::ParseParameterType(
 	FStringView InType,
-	TArrayView<const TCHAR*> InExtraSRVTypes,
-	TArrayView<const TCHAR*> InExtraUAVTypes)
+	TArrayView<const TCHAR* const> InExtraSRVTypes,
+	TArrayView<const TCHAR* const> InExtraUAVTypes)
 {
 	TArrayView<const TCHAR* const> AllSamplerTypes(s_AllSamplerTypes);
 	TArrayView<const TCHAR* const> AllSRVTypes(s_AllSRVTypes);
diff --git a/Engine/Source/Developer/ShaderCompilerCommon/Private/ShaderParameterParser.cpp b/Engine/Source/Developer/ShaderCompilerCommon/Private/ShaderParameterParser.cpp
--- a/Engine/Source/Developer/ShaderCompilerCommon/Private/ShaderParameterParser.cpp
+++ b/Engine/Source/Developer/ShaderCompilerCommon/Private/ShaderParameterParser.cpp
@@ -135,11 +135,18 @@ static void AddNoteToDisplayShaderParameterMemberOnCppSide(
 FShaderParameterParser::FShaderParameterParser() = default;
 FShaderParameterParser::~FShaderParameterParser() = default;
 
+FShaderParameterParser::FShaderParameterParser(const TCHAR* InConstantBufferType)
+	: ConstantBufferType(InConstantBufferType)
+{
+}
+
 FShaderParameterParser::FShaderParameterParser(
-	TArrayView<const TCHAR*> InExtraSRVTypes,
-	TArrayView<const TCHAR*> InExtraUAVTypes
+	const TCHAR* InConstantBufferType,
+	TArrayView<const TCHAR* const> InExtraSRVTypes,
+	TArrayView<const TCHAR* const> InExtraUAVTypes
 )
-	: ExtraSRVTypes(InExtraSRVTypes)
+	: ConstantBufferType(InConstantBufferType)
+	, ExtraSRVTypes(InExtraSRVTypes)
 	, ExtraUAVTypes(InExtraUAVTypes)
 {
 }
@@ -267,7 +274,7 @@ bool FShaderParameterParser::ParseParameters(
 				const bool bBindlessIndex = (ParsedParameterType == EShaderParameterType::BindlessResourceIndex || ParsedParameterType == EShaderParameterType::BindlessSamplerIndex);
 
 				EBindlessConversionType BindlessConversionType = EBindlessConversionType::None;
-				EShaderParameterType ConstantBufferType = ParsedParameterType;
+				EShaderParameterType ParsedConstantBufferType = ParsedParameterType;
 
 				if ((bBindlessResources || bBindlessSamplers) && ParsedParameterType == EShaderParameterType::LooseData)
 				{
@@ -276,19 +283,19 @@ bool FShaderParameterParser::ParseParameters(
 					if (bBindlessResources && (ParsedParameterType == EShaderParameterType::SRV || ParsedParameterType == EShaderParameterType::UAV))
 					{
 						BindlessConversionType = EBindlessConversionType::Resource;
-						ConstantBufferType = EShaderParameterType::BindlessResourceIndex;
+						ParsedConstantBufferType = EShaderParameterType::BindlessResourceIndex;
 					}
 					else if (bBindlessSamplers && ParsedParameterType == EShaderParameterType::Sampler)
 					{
 						BindlessConversionType = EBindlessConversionType::Sampler;
-						ConstantBufferType = EShaderParameterType::BindlessSamplerIndex;
+						ParsedConstantBufferType = EShaderParameterType::BindlessSamplerIndex;
 					}
 
 					if (BindlessConversionType != EBindlessConversionType::None && Leftovers.Contains(TEXT("register")))
 					{
 						// avoid rewriting hardcoded register assignments
 						BindlessConversionType = EBindlessConversionType::None;
-						ConstantBufferType = ParsedParameterType;
+						ParsedConstantBufferType = ParsedParameterType;
 					}
 				}
 
@@ -324,7 +331,7 @@ bool FShaderParameterParser::ParseParameters(
 
 							if (bMoveToRootConstantBuffer)
 							{
-								ConstantBufferParameterType = ConstantBufferType;
+								ConstantBufferParameterType = ParsedConstantBufferType;
 							}
 						}
 					}
@@ -847,8 +854,7 @@ static const TCHAR* GetConstantSwizzle(uint16 ByteOffset)
 bool FShaderParameterParser::MoveShaderParametersToRootConstantBuffer(
 	const FShaderCompilerInput& CompilerInput,
 	FShaderCompilerOutput& CompilerOutput,
-	FString& PreprocessedShaderSource,
-	const TCHAR* ConstantBufferType)
+	FString& PreprocessedShaderSource)
 {
 	bool bSuccess = true;
 
@@ -938,12 +944,11 @@ bool FShaderParameterParser::MoveShaderParametersToRootConstantBuffer(
 bool FShaderParameterParser::ParseAndModify(
 	const FShaderCompilerInput& CompilerInput,
 	FShaderCompilerOutput& CompilerOutput,
-	FString& PreprocessedShaderSource,
-	const TCHAR* ConstantBufferType)
+	FString& PreprocessedShaderSource)
 {
 	bBindlessResources = CompilerInput.Environment.CompilerFlags.Contains(CFLAG_BindlessResources);
 	bBindlessSamplers = CompilerInput.Environment.CompilerFlags.Contains(CFLAG_BindlessSamplers);
-	bHasRootParameters = (CompilerInput.RootParametersStructure != nullptr);
+	const bool bHasRootParameters = (CompilerInput.RootParametersStructure != nullptr);
 
 	// The shader doesn't have any parameter binding through shader structure, therefore don't do anything.
 	if (!(bBindlessResources || bBindlessSamplers || bHasRootParameters))
@@ -951,7 +956,7 @@ bool FShaderParameterParser::ParseAndModify(
 		return true;
 	}
 
-	bNeedToMoveToRootConstantBuffer = ConstantBufferType != nullptr;
+	bNeedToMoveToRootConstantBuffer = ConstantBufferType != nullptr && (CompilerInput.IsRayTracingShader() || UE::ShaderCompilerCommon::ShouldUseStableConstantBuffer(CompilerInput));
 	OriginalParsedShader = PreprocessedShaderSource;
 
 	if (!ParseParameters(CompilerInput, CompilerOutput))
@@ -964,7 +969,7 @@ bool FShaderParameterParser::ParseAndModify(
 
 	if (bNeedToMoveToRootConstantBuffer)
 	{
-		return MoveShaderParametersToRootConstantBuffer(CompilerInput, CompilerOutput, PreprocessedShaderSource, ConstantBufferType);
+		return MoveShaderParametersToRootConstantBuffer(CompilerInput, CompilerOutput, PreprocessedShaderSource);
 	}
 	
 	return true;
diff --git a/Engine/Source/Developer/ShaderCompilerCommon/Public/ShaderCompilerCommon.h b/Engine/Source/Developer/ShaderCompilerCommon/Public/ShaderCompilerCommon.h
--- a/Engine/Source/Developer/ShaderCompilerCommon/Public/ShaderCompilerCommon.h
+++ b/Engine/Source/Developer/ShaderCompilerCommon/Public/ShaderCompilerCommon.h
@@ -56,7 +56,9 @@ extern SHADERCOMPILERCOMMON_API int16 GetNumUniformBuffersUsed(const FShaderComp
 
 namespace UE::ShaderCompilerCommon
 {
-	extern SHADERCOMPILERCOMMON_API EShaderParameterType ParseParameterType(FStringView InType, TArrayView<const TCHAR*> InExtraSRVTypes, TArrayView<const TCHAR*> InExtraUAVTypes);
+	extern SHADERCOMPILERCOMMON_API bool                 ShouldUseStableConstantBuffer(const FShaderCompilerInput& Input);
+
+	extern SHADERCOMPILERCOMMON_API EShaderParameterType ParseParameterType(FStringView InType, TArrayView<const TCHAR* const> InExtraSRVTypes, TArrayView<const TCHAR* const> InExtraUAVTypes);
 
 	extern SHADERCOMPILERCOMMON_API FStringView          RemoveConstantBufferPrefix(FStringView InName);
 	extern SHADERCOMPILERCOMMON_API FString              RemoveConstantBufferPrefix(const FString& InName);
diff --git a/Engine/Source/Developer/ShaderCompilerCommon/Public/ShaderParameterParser.h b/Engine/Source/Developer/ShaderCompilerCommon/Public/ShaderParameterParser.h
--- a/Engine/Source/Developer/ShaderCompilerCommon/Public/ShaderParameterParser.h
+++ b/Engine/Source/Developer/ShaderCompilerCommon/Public/ShaderParameterParser.h
@@ -64,26 +64,37 @@ public:
 	};
 
 	FShaderParameterParser();
-	FShaderParameterParser(TArrayView<const TCHAR*> InExtraSRVTypes, TArrayView<const TCHAR*> InExtraUAVTypes);
+	FShaderParameterParser(const TCHAR* InConstantBufferType);
+	FShaderParameterParser(const TCHAR* InConstantBufferType, TArrayView<const TCHAR* const> InExtraSRVTypes, TArrayView<const TCHAR* const> InExtraUAVTypes);
 	~FShaderParameterParser();
 
 	/** Parses the preprocessed shader code and applies the necessary modifications to it. */
 	bool ParseAndModify(
 		const FShaderCompilerInput& CompilerInput,
 		FShaderCompilerOutput& CompilerOutput,
-		FString& PreprocessedShaderSource,
-		const TCHAR* ConstantBufferType
+		FString& PreprocessedShaderSource
 	);
 
+	UE_DEPRECATED(5.2, "ParseAndModify doesn't need ConstantBufferType anymore")
+	bool ParseAndModify(
+		const FShaderCompilerInput& CompilerInput,
+		FShaderCompilerOutput& CompilerOutput,
+		FString& PreprocessedShaderSource,
+		const TCHAR* InConstantBufferType
+	)
+	{
+		return ParseAndModify(CompilerInput, CompilerOutput, PreprocessedShaderSource);
+	}
+
 	/** Parses the preprocessed shader code and move the parameters into root constant buffer */
 	UE_DEPRECATED(5.1, "ParseAndModify should be called instead.")
 	bool ParseAndMoveShaderParametersToRootConstantBuffer(
 		const FShaderCompilerInput& CompilerInput,
 		FShaderCompilerOutput& CompilerOutput,
 		FString& PreprocessedShaderSource,
-		const TCHAR* ConstantBufferType)
+		const TCHAR* InConstantBufferType)
 	{
-		return ParseAndModify(CompilerInput, CompilerOutput, PreprocessedShaderSource, ConstantBufferType);
+		return ParseAndModify(CompilerInput, CompilerOutput, PreprocessedShaderSource);
 	}
 
 	/** Gets parsing information from a parameter binding name. */
@@ -150,20 +161,20 @@ private:
 	bool MoveShaderParametersToRootConstantBuffer(
 		const FShaderCompilerInput& CompilerInput,
 		FShaderCompilerOutput& CompilerOutput,
-		FString& PreprocessedShaderSource,
-		const TCHAR* ConstantBufferType
+		FString& PreprocessedShaderSource
 	);
 
 	void ExtractFileAndLine(int32 PragamLineoffset, int32 LineOffset, FString& OutFile, FString& OutLine) const;
 
-	TArray<const TCHAR*> ExtraSRVTypes;
-	TArray<const TCHAR*> ExtraUAVTypes;
+	const TCHAR* const ConstantBufferType = nullptr;
+
+	const TArray<const TCHAR*> ExtraSRVTypes;
+	const TArray<const TCHAR*> ExtraUAVTypes;
 
 	FString OriginalParsedShader;
 
 	TMap<FString, FParsedShaderParameter> ParsedParameters;
 
-	bool bHasRootParameters = false;
 	bool bBindlessResources = false;
 	bool bBindlessSamplers = false;
 
diff --git a/Engine/Source/Developer/ShaderFormatOpenGL/Private/OpenGLShaderCompiler.cpp b/Engine/Source/Developer/ShaderFormatOpenGL/Private/OpenGLShaderCompiler.cpp
--- a/Engine/Source/Developer/ShaderFormatOpenGL/Private/OpenGLShaderCompiler.cpp
+++ b/Engine/Source/Developer/ShaderFormatOpenGL/Private/OpenGLShaderCompiler.cpp
@@ -3336,7 +3336,7 @@ void FOpenGLFrontend::CompileShader(const FShaderCompilerInput& Input, FShaderCo
 	}
 
 	FShaderParameterParser ShaderParameterParser;
-	if (!ShaderParameterParser.ParseAndModify(Input, Output, PreprocessedShader, /* ConstantBufferType = */ nullptr))
+	if (!ShaderParameterParser.ParseAndModify(Input, Output, PreprocessedShader))
 	{
 		// The FShaderParameterParser will add any relevant errors.
 		return;
diff --git a/Engine/Source/Developer/VulkanShaderFormat/Private/VulkanShaderCompiler.cpp b/Engine/Source/Developer/VulkanShaderFormat/Private/VulkanShaderCompiler.cpp
--- a/Engine/Source/Developer/VulkanShaderFormat/Private/VulkanShaderCompiler.cpp
+++ b/Engine/Source/Developer/VulkanShaderFormat/Private/VulkanShaderCompiler.cpp
@@ -2352,7 +2352,7 @@ void DoCompileVulkanShader(const FShaderCompilerInput& Input, FShaderCompilerOut
 	}
 
 	FShaderParameterParser ShaderParameterParser;
-	if (!ShaderParameterParser.ParseAndModify(Input, Output, PreprocessedShaderSource, /* ConstantBufferType = */ nullptr))
+	if (!ShaderParameterParser.ParseAndModify(Input, Output, PreprocessedShaderSource))
 	{
 		// The FShaderParameterParser will add any relevant errors.
 		return;
diff --git a/Engine/Source/Developer/Windows/ShaderFormatD3D/Private/D3DShaderCompiler.cpp b/Engine/Source/Developer/Windows/ShaderFormatD3D/Private/D3DShaderCompiler.cpp
--- a/Engine/Source/Developer/Windows/ShaderFormatD3D/Private/D3DShaderCompiler.cpp
+++ b/Engine/Source/Developer/Windows/ShaderFormatD3D/Private/D3DShaderCompiler.cpp
@@ -1302,10 +1302,8 @@ void CompileD3DShader(const FShaderCompilerInput& Input, FShaderCompilerOutput&
 		}
 	}
 
-	FShaderParameterParser ShaderParameterParser;
-	if (!ShaderParameterParser.ParseAndModify(
-		Input, Output, PreprocessedShaderSource,
-		(Input.IsRayTracingShader() || ShouldUseStableConstantBuffer(Input)) ? TEXT("cbuffer") : nullptr))
+	FShaderParameterParser ShaderParameterParser(TEXT("cbuffer"));
+	if (!ShaderParameterParser.ParseAndModify(Input, Output, PreprocessedShaderSource))
 	{
 		// The FShaderParameterParser will add any relevant errors.
 		return;
diff --git a/Engine/Source/Developer/Windows/ShaderFormatD3D/Private/D3DShaderCompiler.inl b/Engine/Source/Developer/Windows/ShaderFormatD3D/Private/D3DShaderCompiler.inl
--- a/Engine/Source/Developer/Windows/ShaderFormatD3D/Private/D3DShaderCompiler.inl
+++ b/Engine/Source/Developer/Windows/ShaderFormatD3D/Private/D3DShaderCompiler.inl
@@ -316,7 +316,7 @@ template <typename ID3D1xShaderReflection, typename D3D1x_SHADER_DESC, typename
 
 			if (bGlobalCB)
 			{
-				if (ShouldUseStableConstantBuffer(Input))
+				if (UE::ShaderCompilerCommon::ShouldUseStableConstantBuffer(Input))
 				{
 					// Each member found in the global constant buffer means it was not in RootParametersStructure or
 					// it would have been moved by ShaderParameterParser.ParseAndModify().
@@ -360,7 +360,7 @@ template <typename ID3D1xShaderReflection, typename D3D1x_SHADER_DESC, typename
 					}
 				}
 			}
-			else if (bIsRootCB && ShouldUseStableConstantBuffer(Input))
+			else if (bIsRootCB && UE::ShaderCompilerCommon::ShouldUseStableConstantBuffer(Input))
 			{
 				if (CBIndex == FShaderParametersMetadata::kRootCBufferBindingIndex)
 				{
diff --git a/Engine/Source/Runtime/RenderCore/Public/ShaderCompilerCore.h b/Engine/Source/Runtime/RenderCore/Public/ShaderCompilerCore.h
--- a/Engine/Source/Runtime/RenderCore/Public/ShaderCompilerCore.h
+++ b/Engine/Source/Runtime/RenderCore/Public/ShaderCompilerCore.h
@@ -506,7 +506,7 @@ enum class ESCWErrorCode
 	BadInputFile
 };
 
-
+UE_DEPRECATED(5.2, "Functionality has moved to UE::ShaderCompilerCommon::ShouldUseStableConstantBuffer")
 inline bool ShouldUseStableConstantBuffer(const FShaderCompilerInput& Input)
 {
 	// stable constant buffer is for the FShaderParameterBindings::BindForLegacyShaderParameters() code path.
