Consider the minimizer to not have succeeded if it returns an empty result.

diff --git a/Engine/Source/Developer/ShaderCompilerCommon/Public/ShaderMinifier.h b/Engine/Source/Developer/ShaderCompilerCommon/Public/ShaderMinifier.h
--- a/Engine/Source/Developer/ShaderCompilerCommon/Public/ShaderMinifier.h
+++ b/Engine/Source/Developer/ShaderCompilerCommon/Public/ShaderMinifier.h
@@ -41,7 +41,7 @@ struct FMinifiedShader
 
 	bool Success() const
 	{
-		return Diagnostics.Errors.IsEmpty();
+		return Diagnostics.Errors.IsEmpty() && !Code.IsEmpty();
 	}
 };
 
