[shaders] refactor usf/directcompile.txt dump into a utility in ShaderCompilerCommon to reduce code duplication in the shader backends

diff --git a/Engine/Source/Developer/Apple/MetalShaderFormat/Private/MetalShaderCompiler.cpp b/Engine/Source/Developer/Apple/MetalShaderFormat/Private/MetalShaderCompiler.cpp
--- a/Engine/Source/Developer/Apple/MetalShaderFormat/Private/MetalShaderCompiler.cpp
+++ b/Engine/Source/Developer/Apple/MetalShaderFormat/Private/MetalShaderCompiler.cpp
@@ -1007,7 +1007,7 @@ void CompileShader_Metal(const FShaderCompilerInput& _Input,FShaderCompilerOutpu
 		Input.DumpDebugInfoPath = FPaths::GetPath(Input.VirtualSourceFilePath);
 	}
 	
-	const bool bDumpDebugInfo = (Input.DumpDebugInfoPath != TEXT("") && IFileManager::Get().DirectoryExists(*Input.DumpDebugInfoPath));
+	const bool bDumpDebugInfo = Input.DumpDebugInfoEnabled();
 
 	// Allow the shader pipeline to override the platform default in here.
 	uint32 MaxUnrollLoops = 32;
@@ -1093,26 +1093,17 @@ void CompileShader_Metal(const FShaderCompilerInput& _Input,FShaderCompilerOutpu
 	// Write out the preprocessed file and a batch file to compile it if requested (DumpDebugInfoPath is valid)
 	if (bDumpDebugInfo && !bDirectCompile)
 	{
-		FArchive* FileWriter = IFileManager::Get().CreateFileWriter(*(Input.DumpDebugInfoPath / FPaths::GetBaseFilename(Input.GetSourceFilename() + TEXT(".usf"))));
-		if (FileWriter)
+		UE::ShaderCompilerCommon::FDebugShaderDataOptions DebugDataOptions;
+		DebugDataOptions.HlslCCFlags = CCFlags;
+		DebugDataOptions.AppendPostSource = [&Input]()
 		{
-			FString Line = GetDumpDebugUSFContents(Input, PreprocessedShader, 0);
-
 			// add the remote data if necessary
 //			if (IsRemoteBuildingConfigured(&Input.Environment))
 			{
-				Line += CreateRemoteDataFromEnvironment(Input.Environment);
+				return CreateRemoteDataFromEnvironment(Input.Environment);
 			}
-
-			FileWriter->Serialize(TCHAR_TO_ANSI(*Line), Line.Len());
-			FileWriter->Close();
-			delete FileWriter;
-		}
-
-		if (Input.bGenerateDirectCompileFile)
-		{
-			FFileHelper::SaveStringToFile(CreateShaderCompilerWorkerDirectCommandLine(Input, CCFlags), *(Input.DumpDebugInfoPath / TEXT("DirectCompile.txt")));
-		}
+		};
+		UE::ShaderCompilerCommon::DumpDebugShaderData(Input, PreprocessedShader, DebugDataOptions);
 	}
 
 	FSHAHash GUIDHash;
diff --git a/Engine/Source/Developer/ShaderCompilerCommon/Private/ShaderCompilerCommon.cpp b/Engine/Source/Developer/ShaderCompilerCommon/Private/ShaderCompilerCommon.cpp
--- a/Engine/Source/Developer/ShaderCompilerCommon/Private/ShaderCompilerCommon.cpp
+++ b/Engine/Source/Developer/ShaderCompilerCommon/Private/ShaderCompilerCommon.cpp
@@ -2,6 +2,7 @@
 
 #include "ShaderCompilerCommon.h"
 #include "ShaderParameterParser.h"
+#include "Misc/FileHelper.h"
 #include "Misc/Paths.h"
 #include "Modules/ModuleManager.h"
 #include "HlslccDefinitions.h"
@@ -1389,25 +1390,23 @@ void CompileOfflineMali(const FShaderCompilerInput& Input, FShaderCompilerOutput
 
 FString GetDumpDebugUSFContents(const FShaderCompilerInput& Input, const FString& Source, uint32 HlslCCFlags)
 {
-	FString Contents = Source;
-	Contents += TEXT("\n");
-	Contents += CrossCompiler::CreateResourceTableFromEnvironment(Input.Environment);
-	Contents += TEXT("#if 0 /*DIRECT COMPILE*/\n");
-	Contents += CreateShaderCompilerWorkerDirectCommandLine(Input, HlslCCFlags);
-	Contents += TEXT("\n#endif /*DIRECT COMPILE*/\n");
-
-	return Contents;
+	UE::ShaderCompilerCommon::FDebugShaderDataOptions DebugDataOptions;
+	DebugDataOptions.HlslCCFlags = HlslCCFlags;
+	return UE::ShaderCompilerCommon::GetDebugShaderContents(Input, Source, DebugDataOptions);
 }
 
 void DumpDebugUSF(const FShaderCompilerInput& Input, const ANSICHAR* Source, uint32 HlslCCFlags, const TCHAR* OverrideBaseFilename)
 {
+	PRAGMA_DISABLE_DEPRECATION_WARNINGS
 	FString NewSource = Source ? Source : "";
 	FString Contents = GetDumpDebugUSFContents(Input, NewSource, HlslCCFlags);
 	DumpDebugUSF(Input, NewSource, HlslCCFlags, OverrideBaseFilename);
+	PRAGMA_ENABLE_DEPRECATION_WARNINGS
 }
 
 void DumpDebugUSF(const FShaderCompilerInput& Input, const FString& Source, uint32 HlslCCFlags, const TCHAR* OverrideBaseFilename)
 {
+	PRAGMA_DISABLE_DEPRECATION_WARNINGS
 	FString BaseSourceFilename = (OverrideBaseFilename && *OverrideBaseFilename) ? OverrideBaseFilename : *Input.GetSourceFilename();
 	FString Filename = Input.DumpDebugInfoPath / BaseSourceFilename;
 
@@ -1417,6 +1416,65 @@ void DumpDebugUSF(const FShaderCompilerInput& Input, const FString& Source, uint
 		FileWriter->Serialize(TCHAR_TO_ANSI(*Contents), Contents.Len());
 		FileWriter->Close();
 	}
+	PRAGMA_ENABLE_DEPRECATION_WARNINGS
+}
+
+namespace UE::ShaderCompilerCommon
+{
+	void DumpDebugShaderData(const FShaderCompilerInput& Input, const FString& PreprocessedSource, const FDebugShaderDataOptions& Options)
+	{
+		if (Input.DumpDebugInfoEnabled())
+		{
+			FString Prefix = (Options.FilenamePrefix && *Options.FilenamePrefix) ? Options.FilenamePrefix : FString();
+			FString BaseSourceFilename = Prefix + ((Options.OverrideBaseFilename && *Options.OverrideBaseFilename) ? Options.OverrideBaseFilename : *Input.GetSourceFilename());
+
+			FString ModifiedPreprocessedSource = GetDebugShaderContents(Input, PreprocessedSource, Options);
+			PRAGMA_DISABLE_DEPRECATION_WARNINGS
+			// note: DumpDebugUSF should be internalized once the deprecation window ends 
+			// (i.e. still should remain in this file but not be declared in the header)
+			// until that time we need to disable deprecation warnings here.
+			DumpDebugUSF(Input, ModifiedPreprocessedSource, Options.HlslCCFlags, *BaseSourceFilename);
+			PRAGMA_ENABLE_DEPRECATION_WARNINGS
+
+			if (Input.bGenerateDirectCompileFile)
+			{
+				// note: CreateShaderCompileWorkerDirectCommandLine should be internalized once the deprecation window ends
+				// (i.e. still should remain in this file but not be declared in the header)
+				// until that time we need to disable deprecation warnings here.
+				PRAGMA_DISABLE_DEPRECATION_WARNINGS
+				FFileHelper::SaveStringToFile(CreateShaderCompilerWorkerDirectCommandLine(Input), *(Input.DumpDebugInfoPath / FString::Printf(TEXT("%sDirectCompile.txt"), *Prefix)));
+				PRAGMA_ENABLE_DEPRECATION_WARNINGS
+			}
+		}
+	}
+
+	FString GetDebugShaderContents(const FShaderCompilerInput& Input, const FString& PreprocessedSource, const FDebugShaderDataOptions& Options)
+	{
+		FString Contents = Options.AppendPreSource ? Options.AppendPreSource() : FString();
+		Contents += PreprocessedSource;
+		if (Options.AppendPostSource)
+		{
+			Contents += Options.AppendPostSource();
+		}
+		Contents += TEXT("\n");
+		Contents += CrossCompiler::CreateResourceTableFromEnvironment(Input.Environment);
+		Contents += TEXT("#if 0 /*DIRECT COMPILE*/\n");
+		// note: CreateShaderCompileWorkerDirectCommandLine should be internalized once the deprecation window ends
+		// (i.e. still should remain in this file but not be declared in the header)
+		// until that time we need to disable deprecation warnings here.
+		PRAGMA_DISABLE_DEPRECATION_WARNINGS	
+		Contents += CreateShaderCompilerWorkerDirectCommandLine(Input, Options.HlslCCFlags);
+		PRAGMA_ENABLE_DEPRECATION_WARNINGS
+		Contents += TEXT("\n#endif /*DIRECT COMPILE*/\n");
+		if (!Input.DebugDescription.IsEmpty())
+		{
+			Contents += TEXT("//");
+			Contents += Input.DebugDescription;
+			Contents += TEXT("\n");
+		}
+
+		return Contents;
+	}
 }
 
 void DumpDebugShaderText(const FShaderCompilerInput& Input, const FString& InSource, const FString& FileExtension)
diff --git a/Engine/Source/Developer/ShaderCompilerCommon/Public/ShaderCompilerCommon.h b/Engine/Source/Developer/ShaderCompilerCommon/Public/ShaderCompilerCommon.h
--- a/Engine/Source/Developer/ShaderCompilerCommon/Public/ShaderCompilerCommon.h
+++ b/Engine/Source/Developer/ShaderCompilerCommon/Public/ShaderCompilerCommon.h
@@ -10,6 +10,8 @@
 #if UE_ENABLE_INCLUDE_ORDER_DEPRECATED_IN_5_1
 #include "ShaderParameterParser.h"
 #endif
+#include "Templates/Function.h"
+
 
 class FShaderParameterParser;
 
@@ -81,6 +83,36 @@ namespace UE::ShaderCompilerCommon
 	*    NOTE: closesthit attribute must always be provided for complex hit group entry points
 	*/
 	extern SHADERCOMPILERCOMMON_API void ParseRayTracingEntryPoint(const FString& Input, FString& OutMain, FString& OutAnyHit, FString& OutIntersection);
+
+	struct FDebugShaderDataOptions
+	{
+		uint32 HlslCCFlags = 0;
+		const TCHAR* OverrideBaseFilename = nullptr;
+		const TCHAR* FilenamePrefix = nullptr;
+		TFunction<FString()> AppendPreSource{};
+		TFunction<FString()> AppendPostSource{};
+	};
+	/*
+	 * Dumps common debug information (preprocessed .usf as constructed by GetDebugShaderContents, and a directcompile.txt file 
+	 * containing the commandline for launching ShaderCompileWorker manually) for the given shader compile input
+	 * and preprocessed source according to the options provided.
+	 * @param	Input The input of the compilation job
+	 * @param	PreprocessedSource The unmodified preprocessed source (used as input to the compilation)
+	 * @param	Options Options which can change behaviour of the debug dump; see above.
+	 */
+	extern SHADERCOMPILERCOMMON_API void DumpDebugShaderData(const FShaderCompilerInput& Input, const FString& PreprocessedSource, const FDebugShaderDataOptions& Options = FDebugShaderDataOptions());
+
+	/*
+	 * Constructs the modified preprocessed source that would be dumped to a .usf file via DumpDebugShaderData, including the following additions:
+	 * - resource table debug information
+	 * - commandlines for running ShaderCompileWorker in DirectCompile mode (one which re-runs the preprocessor and one which bypasses it)
+	 * - a comment containing the FShaderCompilerInput DebugDescription
+	 * - any additions performed by the AppendPreSource/AppendPostSource TFunctions
+	 * @param	Input The input of the compilation job
+	 * @param	PreprocessedSource The unmodified preprocessed source (used as input to the compilation)
+	 * @param	Options Options which can change behaviour of the debug dump; see above.
+	 */
+	extern SHADERCOMPILERCOMMON_API FString GetDebugShaderContents(const FShaderCompilerInput& Input, const FString& PreprocessedSource, const FDebugShaderDataOptions& Options = FDebugShaderDataOptions());
 }
 
 extern SHADERCOMPILERCOMMON_API void HandleReflectedGlobalConstantBufferMember(
@@ -202,6 +234,7 @@ extern SHADERCOMPILERCOMMON_API bool RemoveUnusedInputs(FString& InOutSourceCode
 
 extern SHADERCOMPILERCOMMON_API bool ConvertFromFP32ToFP16(FString& InOutSourceCode, TArray<FString>& OutErrors);
 
+UE_DEPRECATED(5.2, "This function will become internal; please use the UE::ShaderCompilerCommon::DumpDebugShaderData function for dumping common debug info.")
 extern SHADERCOMPILERCOMMON_API FString CreateShaderCompilerWorkerDirectCommandLine(const FShaderCompilerInput& Input, uint32 CCFlags = 0);
 
 enum class EShaderConductorTarget
@@ -211,11 +244,14 @@ enum class EShaderConductorTarget
 };
  extern SHADERCOMPILERCOMMON_API void WriteShaderConductorCommandLine(const FShaderCompilerInput& Input, const FString& SourceFilename, EShaderConductorTarget Target);
 
- // Gets the string that DumpDebugUSF writes out
+// Gets the string that DumpDebugUSF writes out
+UE_DEPRECATED(5.2, "Please use UE::ShaderCompilerCommon::GetDebugShaderContents")
 extern SHADERCOMPILERCOMMON_API FString GetDumpDebugUSFContents(const FShaderCompilerInput& Input, const FString& Source, uint32 HlslCCFlags);
 
 // Utility functions shared amongst all backends to write out a dumped USF
+UE_DEPRECATED(5.2, "Use DumpDebugShaderData to handle all debug dumps from compiler backends.")
 extern SHADERCOMPILERCOMMON_API void DumpDebugUSF(const FShaderCompilerInput& Input, const ANSICHAR* Source, uint32 HlslCCFlags = 0, const TCHAR* OverrideBaseFilename = nullptr);
+UE_DEPRECATED(5.2, "Use DumpDebugShaderData to handle all debug dumps from compiler backends.")
 extern SHADERCOMPILERCOMMON_API void DumpDebugUSF(const FShaderCompilerInput& Input, const FString& Source, uint32 HlslCCFlags = 0, const TCHAR* OverrideBaseFilename = nullptr);
 
 extern SHADERCOMPILERCOMMON_API void DumpDebugShaderText(const FShaderCompilerInput& Input, const FString& InSource, const FString& FileExtension);
diff --git a/Engine/Source/Developer/ShaderFormatOpenGL/Private/OpenGLShaderCompiler.cpp b/Engine/Source/Developer/ShaderFormatOpenGL/Private/OpenGLShaderCompiler.cpp
--- a/Engine/Source/Developer/ShaderFormatOpenGL/Private/OpenGLShaderCompiler.cpp
+++ b/Engine/Source/Developer/ShaderFormatOpenGL/Private/OpenGLShaderCompiler.cpp
@@ -2996,21 +2996,9 @@ static bool CompileToGlslWithShaderConductor(
 	FShaderCompilerDefinitions AdditionalDefines;
 	AdditionalDefines.SetDefine(TEXT("TextureExternal"), TEXT("Texture2D"));
 
-	if (bDumpDebugInfo)
-	{
-		FString DirectCompileLine = CrossCompiler::CreateResourceTableFromEnvironment(Input.Environment);
-
-		DirectCompileLine += TEXT("#if 0 /*DIRECT COMPILE*/\n");
-		DirectCompileLine += CreateShaderCompilerWorkerDirectCommandLine(Input, CCFlags);
-		DirectCompileLine += TEXT("\n#endif /*DIRECT COMPILE*/\n");
-
-		DumpDebugUSF(Input, (PreprocessedShader + DirectCompileLine), CCFlags);
-
-		if (Input.bGenerateDirectCompileFile)
-		{
-			FFileHelper::SaveStringToFile(CreateShaderCompilerWorkerDirectCommandLine(Input), *(Input.DumpDebugInfoPath / TEXT("DirectCompile.txt")));
-		}
-	}
+	UE::ShaderCompilerCommon::FDebugShaderDataOptions DebugDataOptions;
+	DebugDataOptions.HlslCCFlags = CCFlags;
+	UE::ShaderCompilerCommon::DumpDebugShaderData(Input, PreprocessedShader, DebugDataOptions);
 
 	uint32_t BlendFlags = GetDecalBlendFlags(SourceData);
 
@@ -3363,17 +3351,10 @@ void FOpenGLFrontend::CompileShader(const FShaderCompilerInput& Input, FShaderCo
 	else
 #endif // DXC_SUPPORTED
 	{
-		// Write out the preprocessed file and a batch file to compile it if requested (DumpDebugInfoPath is valid)
-		if (bDumpDebugInfo)
-		{
-			DumpDebugUSF(Input, PreprocessedShader, CCFlags);
-
-			if (Input.bGenerateDirectCompileFile)
-			{
-				FFileHelper::SaveStringToFile(CreateShaderCompilerWorkerDirectCommandLine(Input), *(Input.DumpDebugInfoPath / TEXT("DirectCompile.txt")));
-			}
-		}
-
+		UE::ShaderCompilerCommon::FDebugShaderDataOptions DebugDataOptions;
+		DebugDataOptions.HlslCCFlags = CCFlags;
+		UE::ShaderCompilerCommon::DumpDebugShaderData(Input, PreprocessedShader, DebugDataOptions);
+		
 		CCFlags |= HLSLCC_NoValidation;
 		FGlslCodeBackend* BackEnd = CreateBackend(Version, CCFlags, HlslCompilerTarget);
 
diff --git a/Engine/Source/Developer/ShaderFormatVectorVM/Private/VectorVMShaderCompiler.cpp b/Engine/Source/Developer/ShaderFormatVectorVM/Private/VectorVMShaderCompiler.cpp
--- a/Engine/Source/Developer/ShaderFormatVectorVM/Private/VectorVMShaderCompiler.cpp
+++ b/Engine/Source/Developer/ShaderFormatVectorVM/Private/VectorVMShaderCompiler.cpp
@@ -102,28 +102,6 @@ bool CompileShader_VectorVM(const FShaderCompilerInput& Input, FShaderCompilerOu
 
 	const EHlslShaderFrequency Frequency = HSF_VertexShader;
 
- 		// Write out the preprocessed file and a batch file to compile it if requested (DumpDebugInfoPath is valid)
-	if (bDumpDebugInfo)
-	{
-		FArchive* FileWriter = IFileManager::Get().CreateFileWriter(*(Input.DumpDebugInfoPath / Input.GetSourceFilename()));
-		if (FileWriter)
-		{
-			auto AnsiSourceFile = StringCast<ANSICHAR>(*PreprocessedShader);
-			FileWriter->Serialize((ANSICHAR*)AnsiSourceFile.Get(), AnsiSourceFile.Length());
-			{
-				FString Line = CrossCompiler::CreateResourceTableFromEnvironment(Input.Environment);
-				FileWriter->Serialize(TCHAR_TO_ANSI(*Line), Line.Len());
-			}
-			FileWriter->Close();
-			delete FileWriter;
-		}
-
-		if (Input.bGenerateDirectCompileFile)
-		{
-			FFileHelper::SaveStringToFile(CreateShaderCompilerWorkerDirectCommandLine(Input), *(Input.DumpDebugInfoPath / TEXT("DirectCompile.txt")));
-		}
-	}
-
 	//Is stuff like this needed? What others?
 	uint32 CCFlags = HLSLCC_NoPreprocess;
 		//CCFlags |= HLSLCC_PrintAST;
@@ -151,6 +129,10 @@ bool CompileShader_VectorVM(const FShaderCompilerInput& Input, FShaderCompilerOu
 	// Required as we added the RemoveUniformBuffersFromSource() function (the cross-compiler won't be able to interpret comments w/o a preprocessor)
 	//CCFlags &= ~HLSLCC_NoPreprocess;
 
+	UE::ShaderCompilerCommon::FDebugShaderDataOptions DebugDataOptions;
+	DebugDataOptions.HlslCCFlags = CCFlags;
+	UE::ShaderCompilerCommon::DumpDebugShaderData(Input, PreprocessedShader, DebugDataOptions);
+
 	FVectorVMCodeBackend VVMBackEnd(CCFlags, HlslCompilerTarget, VMCompilationOutput);
 		FVectorVMLanguageSpec VVMLanguageSpec; 
 
diff --git a/Engine/Source/Developer/VulkanShaderFormat/Private/VulkanShaderCompiler.cpp b/Engine/Source/Developer/VulkanShaderFormat/Private/VulkanShaderCompiler.cpp
--- a/Engine/Source/Developer/VulkanShaderFormat/Private/VulkanShaderCompiler.cpp
+++ b/Engine/Source/Developer/VulkanShaderFormat/Private/VulkanShaderCompiler.cpp
@@ -1348,10 +1348,8 @@ FCompilerInfo::FCompilerInfo(const FShaderCompilerInput& InInput, const FString&
 	Input(InInput),
 	WorkingDirectory(InWorkingDirectory),
 	CCFlags(0),
-	Frequency(InFrequency),
-	bDebugDump(false)
+	Frequency(InFrequency)
 {
-	bDebugDump = Input.DumpDebugInfoPath != TEXT("") && IFileManager::Get().DirectoryExists(*Input.DumpDebugInfoPath);
 	BaseSourceFilename = Input.GetSourceFilename();
 }
 
@@ -2156,7 +2154,7 @@ static bool CompileWithShaderConductor(
 	const bool bIsRayTracingShader = Input.IsRayTracingShader();
 	const bool bHasBindless = Input.Environment.CompilerFlags.Contains(CFLAG_BindlessResources) || Input.Environment.CompilerFlags.Contains(CFLAG_BindlessSamplers);
 	const bool bRewriteHlslSource = !bIsRayTracingShader;
-	const bool bDebugDump = CompilerInfo.bDebugDump;
+	const bool bDebugDump = Input.DumpDebugInfoEnabled();
 
 	CrossCompiler::FShaderConductorContext CompilerContext;
 
@@ -2186,6 +2184,8 @@ static bool CompileWithShaderConductor(
 		Options.TargetEnvironment = GetMinimumTargetEnvironment(Input.Target.GetPlatform());
 	}
 
+	UE::ShaderCompilerCommon::DumpDebugShaderData(Input, PreprocessedShader, { CompilerInfo.CCFlags });
+
 	if (bDebugDump)
 	{
 		VulkanCreateDXCCompileBatchFiles(
@@ -2193,8 +2193,6 @@ static bool CompileWithShaderConductor(
 			Frequency,
 			CompilerInfo,
 			Options);
-
-		DumpDebugUSF(Input, PreprocessedShader, CompilerInfo.CCFlags);
 	}
 
 	if (bRewriteHlslSource)
@@ -2406,11 +2404,7 @@ void DoCompileVulkanShader(const FShaderCompilerInput& Input, FShaderCompilerOut
 		}
 	}
 
-	// Write out the preprocessed file and a batch file to compile it if requested (DumpDebugInfoPath is valid)
-	if (CompilerInfo.bDebugDump)
-	{
-		DumpDebugUSF(Input, PreprocessedShaderSource, CompilerInfo.CCFlags);
-	}
+	UE::ShaderCompilerCommon::DumpDebugShaderData(Input, PreprocessedShaderSource, { CompilerInfo.CCFlags });
 
 	TArray<ANSICHAR> GeneratedGlslSource;
 	FVulkanBindingTable BindingTable(CompilerInfo.Frequency);
diff --git a/Engine/Source/Developer/VulkanShaderFormat/Private/VulkanShaderFormat.h b/Engine/Source/Developer/VulkanShaderFormat/Private/VulkanShaderFormat.h
--- a/Engine/Source/Developer/VulkanShaderFormat/Private/VulkanShaderFormat.h
+++ b/Engine/Source/Developer/VulkanShaderFormat/Private/VulkanShaderFormat.h
@@ -24,7 +24,6 @@ struct FCompilerInfo
 	FString Profile;
 	uint32 CCFlags;
 	EHlslShaderFrequency Frequency;
-	bool bDebugDump;
 	FString BaseSourceFilename;
 
 	FCompilerInfo(const struct FShaderCompilerInput& InInput, const FString& InWorkingDirectory, EHlslShaderFrequency InFrequency);
diff --git a/Engine/Source/Developer/Windows/ShaderFormatD3D/Private/D3DShaderCompiler.cpp b/Engine/Source/Developer/Windows/ShaderFormatD3D/Private/D3DShaderCompiler.cpp
--- a/Engine/Source/Developer/Windows/ShaderFormatD3D/Private/D3DShaderCompiler.cpp
+++ b/Engine/Source/Developer/Windows/ShaderFormatD3D/Private/D3DShaderCompiler.cpp
@@ -498,40 +498,6 @@ inline bool IsCompatibleBinding(const D3D11_SHADER_INPUT_BIND_DESC& BindDesc, ui
 	return true;
 }
 
-bool DumpDebugShaderUSF(FString& PreprocessedShaderSource, const FShaderCompilerInput& Input)
-{
-	bool bDumpDebugInfo = false;
-
-	// Write out the preprocessed file and a batch file to compile it if requested (DumpDebugInfoPath is valid)
-	if (Input.DumpDebugInfoPath.Len() > 0 && IFileManager::Get().DirectoryExists(*Input.DumpDebugInfoPath))
-	{
-		bDumpDebugInfo = true;
-		FString Filename = Input.GetSourceFilename();
-		FArchive* FileWriter = IFileManager::Get().CreateFileWriter(*(Input.DumpDebugInfoPath / Filename));
-		if (FileWriter)
-		{
-			auto AnsiSourceFile = StringCast<ANSICHAR>(*PreprocessedShaderSource);
-			FileWriter->Serialize((ANSICHAR*)AnsiSourceFile.Get(), AnsiSourceFile.Length());
-			{
-				FString Line = CrossCompiler::CreateResourceTableFromEnvironment(Input.Environment);
-
-				Line += TEXT("#if 0 /*DIRECT COMPILE*/\n");
-				Line += CreateShaderCompilerWorkerDirectCommandLine(Input);
-				Line += TEXT("\n#endif /*DIRECT COMPILE*/\n");
-				Line += TEXT("//");
-				Line += Input.DebugDescription;
-				Line += TEXT("\n");
-				FileWriter->Serialize(TCHAR_TO_ANSI(*Line), Line.Len());
-			}
-			FileWriter->Close();
-			delete FileWriter;
-		}
-	}
-
-	return bDumpDebugInfo;
-}
-
-
 static void PatchSpirvForPrecompilation(FSpirv& Spirv)
 {
 	// Remove [unroll] loop hints from SPIR-V as this can fail on infinite loops
@@ -622,7 +588,8 @@ bool CompileAndProcessD3DShaderFXC(FString& PreprocessedShaderSource, const FStr
 	auto AnsiSourceFile = StringCast<ANSICHAR>(*PreprocessedShaderSource);
 
 	// Write out the preprocessed file and a batch file to compile it if requested (DumpDebugInfoPath is valid)
-	bool bDumpDebugInfo = DumpDebugShaderUSF(PreprocessedShaderSource, Input);
+	UE::ShaderCompilerCommon::DumpDebugShaderData(Input, PreprocessedShaderSource);
+	bool bDumpDebugInfo = Input.DumpDebugInfoEnabled();
 	FString DisasmFilename;
 	if (bDumpDebugInfo)
 	{
@@ -638,12 +605,6 @@ bool CompileAndProcessD3DShaderFXC(FString& PreprocessedShaderSource, const FStr
 
 		FFileHelper::SaveStringToFile(BatchFileContents, *(Input.DumpDebugInfoPath / TEXT("CompileFXC.bat")));
 
-		if (Input.bGenerateDirectCompileFile)
-		{
-			FFileHelper::SaveStringToFile(CreateShaderCompilerWorkerDirectCommandLine(Input), *(Input.DumpDebugInfoPath / TEXT("DirectCompile.txt")));
-			FFileHelper::SaveStringToFile(Input.DebugDescription, *(Input.DumpDebugInfoPath / TEXT("permutation_info.txt")));
-		}
-
 		DisasmFilename = *(Input.DumpDebugInfoPath / TEXT("Output.d3dasm"));
 	}
 
@@ -1204,7 +1165,7 @@ void CompileD3DShader(const FShaderCompilerInput& Input, FShaderCompilerOutput&
 		Exceptions.AddUnique(TEXT("SV_CullDistance7"));
 		
 		// Write the preprocessed file out in case so we can debug issues on HlslParser
-		DumpDebugShaderUSF(PreprocessedShaderSource, Input);
+		UE::ShaderCompilerCommon::DumpDebugShaderData(Input, PreprocessedShaderSource);
 
 		TArray<FString> Errors;
 		if (!RemoveUnusedOutputs(PreprocessedShaderSource, UsedOutputs, Exceptions, EntryPointName, Errors))
diff --git a/Engine/Source/Developer/Windows/ShaderFormatD3D/Private/D3DShaderCompilerDXC.cpp b/Engine/Source/Developer/Windows/ShaderFormatD3D/Private/D3DShaderCompilerDXC.cpp
--- a/Engine/Source/Developer/Windows/ShaderFormatD3D/Private/D3DShaderCompilerDXC.cpp
+++ b/Engine/Source/Developer/Windows/ShaderFormatD3D/Private/D3DShaderCompilerDXC.cpp
@@ -552,7 +552,9 @@ static bool RewriteUsingSC(FString& PreprocessedShaderSource, const FShaderCompi
 
 			if (bDumpDebugInfo)
 			{
-				DumpDebugUSF(Input, CStrSourceData.c_str(), 0, GRewrittenBaseFilename);
+				UE::ShaderCompilerCommon::FDebugShaderDataOptions DebugDataOptions;
+				DebugDataOptions.OverrideBaseFilename = GRewrittenBaseFilename;
+				UE::ShaderCompilerCommon::DumpDebugShaderData(Input, PreprocessedShaderSource, DebugDataOptions);
 			}
 		}
 	}
@@ -611,8 +613,8 @@ bool CompileAndProcessD3DShaderDXC(FString& PreprocessedShaderSource,
 		bEnable16BitTypes = true;
 	}
 
-	// Write out the preprocessed file and a batch file to compile it if requested (DumpDebugInfoPath is valid)
-	bool bDumpDebugInfo = DumpDebugShaderUSF(PreprocessedShaderSource, Input);
+	UE::ShaderCompilerCommon::DumpDebugShaderData(Input, PreprocessedShaderSource);
+	bool bDumpDebugInfo = Input.DumpDebugInfoEnabled(); 
 
 	FString Filename = Input.GetSourceFilename();
 
@@ -665,12 +667,6 @@ bool CompileAndProcessD3DShaderDXC(FString& PreprocessedShaderSource,
 	{
 		FString BatchFileContents = D3DCreateDXCCompileBatchFile(Args, Filename);
 		FFileHelper::SaveStringToFile(BatchFileContents, *(Input.DumpDebugInfoPath / TEXT("CompileDXC.bat")));
-
-		if (Input.bGenerateDirectCompileFile)
-		{
-			FFileHelper::SaveStringToFile(CreateShaderCompilerWorkerDirectCommandLine(Input), *(Input.DumpDebugInfoPath / TEXT("DirectCompile.txt")));
-			FFileHelper::SaveStringToFile(Input.DebugDescription, *(Input.DumpDebugInfoPath / TEXT("permutation_info.txt")));
-		}
 	}
 
 	TRefCountPtr<IDxcBlob> ShaderBlob;
diff --git a/Engine/Source/Developer/Windows/ShaderFormatD3D/Private/ShaderFormatD3D.h b/Engine/Source/Developer/Windows/ShaderFormatD3D/Private/ShaderFormatD3D.h
--- a/Engine/Source/Developer/Windows/ShaderFormatD3D/Private/ShaderFormatD3D.h
+++ b/Engine/Source/Developer/Windows/ShaderFormatD3D/Private/ShaderFormatD3D.h
@@ -39,6 +39,5 @@ bool CompileAndProcessD3DShaderDXC(FString& PreprocessedShaderSource,
 	TArray<FString>& FilteredErrors, FShaderCompilerOutput& Output);
 
 bool ValidateResourceCounts(uint32 NumSRVs, uint32 NumSamplers, uint32 NumUAVs, uint32 NumCBs, TArray<FString>& OutFilteredErrors);
-bool DumpDebugShaderUSF(FString& PreprocessedShaderSource, const FShaderCompilerInput& Input);
 
 extern int32 GD3DAllowRemoveUnused;
diff --git a/Engine/Source/Runtime/RenderCore/Public/ShaderCompilerCore.h b/Engine/Source/Runtime/RenderCore/Public/ShaderCompilerCore.h
--- a/Engine/Source/Runtime/RenderCore/Public/ShaderCompilerCore.h
+++ b/Engine/Source/Runtime/RenderCore/Public/ShaderCompilerCore.h
@@ -7,6 +7,7 @@
 #pragma once
 
 #include "CoreMinimal.h"
+#include "HAL/FileManager.h"
 #include "Stats/Stats.h"
 #include "Templates/RefCounting.h"
 #include "Misc/SecureHash.h"
@@ -119,6 +120,8 @@ enum ECompilerFlags
 	CFLAG_BindlessSamplers,
 	// EXPERIMENTAL: Run the shader re-writer that removes any unused functions/resources/types from source code before compilation.
 	CFLAG_RemoveDeadCode,
+	// EXPERIMENTAL: Use the DXC preprocessor (instead of the legacy MCPP)
+	CFLAG_PreprocessWithDXC,
 
 	CFLAG_Max,
 };
@@ -237,6 +240,11 @@ struct FShaderCompilerInput
 	{
 	}
 
+	bool DumpDebugInfoEnabled() const 
+	{
+		return DumpDebugInfoPath != TEXT("") && IFileManager::Get().DirectoryExists(*DumpDebugInfoPath);
+	}
+
 	// generate human readable name for debugging
 	FString GenerateShaderName() const
 	{
