Add UE::ShaderCompilerCommon::RemoveDeadCode() high-level wrapper for shader minifier, to be consistently used in all shader compiler back-ends

diff --git a/Engine/Source/Developer/ShaderCompilerCommon/Private/ShaderCompilerCommon.cpp b/Engine/Source/Developer/ShaderCompilerCommon/Private/ShaderCompilerCommon.cpp
--- a/Engine/Source/Developer/ShaderCompilerCommon/Private/ShaderCompilerCommon.cpp
+++ b/Engine/Source/Developer/ShaderCompilerCommon/Private/ShaderCompilerCommon.cpp
@@ -8,6 +8,7 @@
 #include "HlslccDefinitions.h"
 #include "HAL/FileManager.h"
 #include "String/RemoveFrom.h"
+#include "ShaderMinifier.h"
 
 IMPLEMENT_MODULE(FDefaultModuleImpl, ShaderCompilerCommon);
 
@@ -670,6 +671,42 @@ void UE::ShaderCompilerCommon::ParseRayTracingEntryPoint(const FString& Input, F
 	}
 }
 
+bool UE::ShaderCompilerCommon::RemoveDeadCode(FString& InOutPreprocessedShaderSource, const FString& EntryPoint, TArray<FShaderCompilerError>& OutErrors)
+{
+	FString EntryMain;
+	FString EntryAnyHit;
+	FString EntryIntersection;
+	UE::ShaderCompilerCommon::ParseRayTracingEntryPoint(EntryPoint, EntryMain, EntryAnyHit, EntryIntersection);
+
+	if (!EntryAnyHit.IsEmpty())
+	{
+		EntryMain += TEXT(";");
+		EntryMain += EntryAnyHit;
+	}
+
+	if (!EntryIntersection.IsEmpty())
+	{
+		EntryMain += TEXT(";");
+		EntryMain += EntryIntersection;
+	}
+
+	UE::ShaderMinifier::FMinifiedShader Minified  = UE::ShaderMinifier::Minify(InOutPreprocessedShaderSource, EntryMain, 
+		UE::ShaderMinifier::EMinifyShaderFlags::OutputReasons
+		| UE::ShaderMinifier::EMinifyShaderFlags::OutputStats
+		| UE::ShaderMinifier::EMinifyShaderFlags::OutputLines);
+
+	if (Minified.Success())
+	{
+		Swap(InOutPreprocessedShaderSource, Minified.Code);
+		return true;
+	}
+	else
+	{
+		OutErrors.Add(TEXT("Shader minification failed."));
+		return false;
+	}
+}
+
 void HandleReflectedGlobalConstantBufferMember(
 	const FString& InMemberName,
 	uint32 ConstantBufferIndex,
diff --git a/Engine/Source/Developer/ShaderCompilerCommon/Public/ShaderCompilerCommon.h b/Engine/Source/Developer/ShaderCompilerCommon/Public/ShaderCompilerCommon.h
--- a/Engine/Source/Developer/ShaderCompilerCommon/Public/ShaderCompilerCommon.h
+++ b/Engine/Source/Developer/ShaderCompilerCommon/Public/ShaderCompilerCommon.h
@@ -84,6 +84,14 @@ namespace UE::ShaderCompilerCommon
 	*/
 	extern SHADERCOMPILERCOMMON_API void ParseRayTracingEntryPoint(const FString& Input, FString& OutMain, FString& OutAnyHit, FString& OutIntersection);
 
+	/**
+	* Rewrites a fully preprocessed shader source code, removing any functions or structs that are not reachable from a given entry point.
+	* This is a high-level wrapper for UE::ShaderMinifier that should be used in all shader compiler back-ends for consistency.
+	* Any errors encountered during parsing or minification are added to OutErrors.
+	* InOutPreprocessedShaderSource is replaced with the rewritten code on success and is kept intact on failure.
+	*/
+	extern SHADERCOMPILERCOMMON_API bool RemoveDeadCode(FString& InOutPreprocessedShaderSource, const FString& EntryPoint, TArray<FShaderCompilerError>& OutErrors);
+
 	struct FDebugShaderDataOptions
 	{
 		uint32 HlslCCFlags = 0;
diff --git a/Engine/Source/Developer/Windows/ShaderFormatD3D/Private/D3DShaderCompiler.cpp b/Engine/Source/Developer/Windows/ShaderFormatD3D/Private/D3DShaderCompiler.cpp
--- a/Engine/Source/Developer/Windows/ShaderFormatD3D/Private/D3DShaderCompiler.cpp
+++ b/Engine/Source/Developer/Windows/ShaderFormatD3D/Private/D3DShaderCompiler.cpp
@@ -1207,36 +1207,7 @@ void CompileD3DShader(const FShaderCompilerInput& Input, FShaderCompilerOutput&
 
 	if (Input.Environment.CompilerFlags.Contains(CFLAG_RemoveDeadCode))
 	{
-		FString EntryMain;
-		FString EntryAnyHit;
-		FString EntryIntersection;
-		UE::ShaderCompilerCommon::ParseRayTracingEntryPoint(EntryPointName, EntryMain, EntryAnyHit, EntryIntersection);
-
-		if (!EntryAnyHit.IsEmpty())
-		{
-			EntryMain += TEXT(";");
-			EntryMain += EntryAnyHit;
-		}
-
-		if (!EntryIntersection.IsEmpty())
-		{
-			EntryMain += TEXT(";");
-			EntryMain += EntryIntersection;
-		}
-
-		UE::ShaderMinifier::FMinifiedShader Minified  = UE::ShaderMinifier::Minify(PreprocessedShaderSource, EntryMain, 
-			UE::ShaderMinifier::EMinifyShaderFlags::OutputReasons
-			| UE::ShaderMinifier::EMinifyShaderFlags::OutputStats
-			| UE::ShaderMinifier::EMinifyShaderFlags::OutputLines);
-
-		if (Minified.Success())
-		{
-			Swap(PreprocessedShaderSource, Minified.Code);
-		}
-		else
-		{
-			FilteredErrors.Add(TEXT("Shader minification failed."));
-		}
+		UE::ShaderCompilerCommon::RemoveDeadCode(PreprocessedShaderSource, EntryPointName, Output.Errors);
 	}
 
 	// @TODO - implement different material path to allow us to remove backwards compat flag on sm5 shaders
