ShaderMinifier - Add basic test for MinifyShader, change SkipUntilNextLine to exclude new line character

diff --git a/Engine/Source/Developer/ShaderCompilerCommon/Private/ShaderMinifier.cpp b/Engine/Source/Developer/ShaderCompilerCommon/Private/ShaderMinifier.cpp
--- a/Engine/Source/Developer/ShaderCompilerCommon/Private/ShaderMinifier.cpp
+++ b/Engine/Source/Developer/ShaderCompilerCommon/Private/ShaderMinifier.cpp
@@ -154,7 +154,6 @@ static FStringView SkipUntilNextLine(FStringView Source)
 	int32 Index = INDEX_NONE;
 	if (Source.FindChar('\n', Index))
 	{
-		Index += 1; // Skip the new line character itself
 		return FStringView(Source.GetData() + Index, Source.Len() - Index);
 	}
 	else
@@ -1677,9 +1676,114 @@ IMPLEMENT_SIMPLE_AUTOMATION_TEST(FShaderMinifierTest, "System.Shaders.ShaderMini
 
 bool FShaderMinifierTest::RunTest(const FString& Parameters)
 {
-	// TODO: run minifier on some tricky examples
+	using namespace UE::ShaderMinifier;
+
+	FStringView TestShaderCode = 
+		TEXT(R"(// dxc /T cs_6_6 /E MainCS MinifierTest.hlsl 
+struct FFoo
+{
+	float X;
+	float Y;
+};
+
+#pragma test_pragma
+struct FBar
+{
+	FFoo Foo;
+};
+
+uint GUnreferencedParameter;
+
+struct FUnreferencedStruct
+{
+	uint X;
+};
+
+uint UnreferencedFunction()
+{
+	return GUnreferencedParameter;
+}
+
+#define COMPILER_DEFINITION_TEST 123
+float Sum(in FBar Param)
+{
+	return Param.Foo.X + Param.Foo.Y;
+}
+
+float FunA()
+{
+	// Comment inside function
+	FBar Temp;
+	Temp.Foo.X = 1;
+	Temp.Foo.Y = 2;
+	return Sum(Temp);
+}
+
+float FunB()
+{
+	return FunA();
+}
+
+#line 1000 "MinifierTest.hlsl"
+// Test comment 1
+void EmptyFunction(){}
+
+// Test comment 2
+[numthreads(1,1,1)]
+// Comment during function declaration
+void MainCS()
+{
+	FunB();
+}
+)");
+
+	auto ChunkPresent = [](const FParsedShader& Parsed, FStringView Name)
+	{
+		for (const FCodeChunk& Chunk : Parsed.Chunks)
+		{
+			for (const FCodeBlock& Block : Chunk.Blocks)
+			{
+				if (Block.Code == Name)
+				{
+					return true;
+				}
+			}
+		}
+		return false;
+	};
+
+	{
+		FParsedShader Parsed = ParseShader(TestShaderCode);
+		FDiagnostics Diagnostics;
+		FString Minified = MinifyShader(Parsed, TEXT("EmptyFunction"), EMinifyShaderFlags::None, Diagnostics);
+		FParsedShader MinifiedParsed = ParseShader(Minified);
+		if (TestEqual(TEXT("MinifyShader: EmptyFunction: num chunks"), MinifiedParsed.Chunks.Num(), 3))
+		{
+			TestEqual(TEXT("MinifyShader: EmptyFunction: pragma"), *FString(MinifiedParsed.Chunks[0].GetCode()), TEXT("#pragma test_pragma"));
+			TestEqual(TEXT("MinifyShader: EmptyFunction: define"), *FString(MinifiedParsed.Chunks[1].GetCode()), TEXT("#define COMPILER_DEFINITION_TEST 123"));
+			TestEqual(TEXT("MinifyShader: EmptyFunction: function"), *FString(MinifiedParsed.Chunks[2].GetCode()), TEXT("void EmptyFunction(){}"));
+		}
+	}
+
+	{
+		FParsedShader Parsed = ParseShader(TestShaderCode);
+		FDiagnostics Diagnostics;
+		FString Minified = MinifyShader(Parsed, TEXT("MainCS"), EMinifyShaderFlags::OutputReasons, Diagnostics);
+		FParsedShader MinifiedParsed = ParseShader(Minified);
+
+		TestTrue(TEXT("MinifyShader: MainCS: contains MainCS"), ChunkPresent(MinifiedParsed, TEXT("MainCS")));
+		TestTrue(TEXT("MinifyShader: MainCS: contains FFoo"), ChunkPresent(MinifiedParsed, TEXT("FFoo")));
+		TestTrue(TEXT("MinifyShader: MainCS: contains FBar"), ChunkPresent(MinifiedParsed, TEXT("FBar")));
+		TestTrue(TEXT("MinifyShader: MainCS: contains Sum"), ChunkPresent(MinifiedParsed, TEXT("Sum")));
+		TestTrue(TEXT("MinifyShader: MainCS: contains FunA"), ChunkPresent(MinifiedParsed, TEXT("FunA")));
+		TestTrue(TEXT("MinifyShader: MainCS: contains FunB"), ChunkPresent(MinifiedParsed, TEXT("FunB")));
+		TestFalse(TEXT("MinifyShader: MainCS: contains UnreferencedFunction"), ChunkPresent(MinifiedParsed, TEXT("UnreferencedFunction")));
+		TestFalse(TEXT("MinifyShader: MainCS: contains FUnreferencedStruct"), ChunkPresent(MinifiedParsed, TEXT("FUnreferencedStruct")));
+		TestFalse(TEXT("MinifyShader: MainCS: contains GUnreferencedParameter"), ChunkPresent(MinifiedParsed, TEXT("GUnreferencedParameter")));
+	}
 
 	int32 NumErrors = ExecutionInfo.GetErrorTotal();
+
 	return NumErrors == 0;
 }
 
