ShaderMinifier - Add an option to preserve comment lines

diff --git a/Engine/Source/Developer/ShaderCompilerCommon/Private/ShaderCompilerCommon.cpp b/Engine/Source/Developer/ShaderCompilerCommon/Private/ShaderCompilerCommon.cpp
--- a/Engine/Source/Developer/ShaderCompilerCommon/Private/ShaderCompilerCommon.cpp
+++ b/Engine/Source/Developer/ShaderCompilerCommon/Private/ShaderCompilerCommon.cpp
@@ -692,6 +692,7 @@ bool UE::ShaderCompilerCommon::RemoveDeadCode(FString& InOutPreprocessedShaderSo
 
 	UE::ShaderMinifier::FMinifiedShader Minified  = UE::ShaderMinifier::Minify(InOutPreprocessedShaderSource, EntryMain, 
 		UE::ShaderMinifier::EMinifyShaderFlags::OutputReasons
+		| UE::ShaderMinifier::EMinifyShaderFlags::OutputCommentLines
 		| UE::ShaderMinifier::EMinifyShaderFlags::OutputStats
 		| UE::ShaderMinifier::EMinifyShaderFlags::OutputLines);
 
diff --git a/Engine/Source/Developer/ShaderCompilerCommon/Private/ShaderMinifier.cpp b/Engine/Source/Developer/ShaderCompilerCommon/Private/ShaderMinifier.cpp
--- a/Engine/Source/Developer/ShaderCompilerCommon/Private/ShaderMinifier.cpp
+++ b/Engine/Source/Developer/ShaderCompilerCommon/Private/ShaderMinifier.cpp
@@ -258,6 +258,7 @@ enum class ECodeChunkType {
 	Enum,
 	Define,
 	Pragma,
+	CommentLine, // Single line comment
 };
 
 struct FCodeChunk
@@ -495,7 +496,24 @@ static FParsedShader ParseShader(FStringView InSource, FDiagnostics& Output)
 			break;
 		}
 
-		if (Source.StartsWith(TEXT("//")) || Source.StartsWith(TEXT("#line")))
+		if (Source.StartsWith(TEXT("//")))
+		{
+			FStringView Remainder = SkipUntilNextLine(Source);
+
+			// Save comment lines that are outside of blocks
+			if (PendingBlocks.IsEmpty())
+			{
+				FStringView Block = SubStrView(Source, 0, Source.Len() - Remainder.Len());
+				AddBlock(EBlockType::Unknown, Block);
+				ChunkType = ECodeChunkType::CommentLine;
+				FinalizeChunk();
+			}
+
+			Source = Remainder;
+
+			continue;
+		}
+		else if (Source.StartsWith(TEXT("#line")))
 		{
 			Source = SkipUntilNextLine(Source);
 			continue;
@@ -900,7 +918,8 @@ static void OutputChunk(const FCodeChunk& Chunk, FStringBuilderBase& OutputStrea
 	if (Chunk.Type != ECodeChunkType::Function
 		&& Chunk.Type != ECodeChunkType::CBuffer
 		&& Chunk.Type != ECodeChunkType::Pragma
-		&& Chunk.Type != ECodeChunkType::Define)
+		&& Chunk.Type != ECodeChunkType::Define
+		&& Chunk.Type != ECodeChunkType::CommentLine)
 	{
 		OutputStream << ";";
 	}
@@ -1320,9 +1339,29 @@ static FString MinifyShader(const FParsedShader& Parsed, FStringView SemicolonSe
 
 	for (const FCodeChunk& Chunk : Parsed.Chunks)
 	{
-		if (Chunk.Type != ECodeChunkType::Pragma      // Pragmas and defines that remain after preprocessing
-			&& Chunk.Type != ECodeChunkType::Define   // must be preserved as they may control important compiler behaviors.
-			&& RelevantChunks.Find(&Chunk) == nullptr)
+		auto ShouldSkipChunk = [&RelevantChunks, &Chunk, Flags]()
+		{
+			// Pragmas and defines that remain after preprocessing must be preserved as they may control important compiler behaviors.
+			if (Chunk.Type == ECodeChunkType::Pragma || Chunk.Type == ECodeChunkType::Define)
+			{
+				return false;
+			}
+
+			// The preprocessed shader may have auto-generated comments such as `// #define FOO 123` that may be useful to keep for debugging.
+			if (Chunk.Type == ECodeChunkType::CommentLine && EnumHasAnyFlags(Flags, EMinifyShaderFlags::OutputCommentLines))
+			{
+				return false;
+			}
+
+			if (RelevantChunks.Find(&Chunk))
+			{
+				return false;
+			}
+
+			return true;
+		};
+
+		if (ShouldSkipChunk())
 		{
 			continue;
 		}
diff --git a/Engine/Source/Developer/ShaderCompilerCommon/Public/ShaderMinifier.h b/Engine/Source/Developer/ShaderCompilerCommon/Public/ShaderMinifier.h
--- a/Engine/Source/Developer/ShaderCompilerCommon/Public/ShaderMinifier.h
+++ b/Engine/Source/Developer/ShaderCompilerCommon/Public/ShaderMinifier.h
@@ -13,9 +13,18 @@ namespace UE::ShaderMinifier
 enum EMinifyShaderFlags
 {
 	None = 0,
-	OutputReasons = 1 << 1,
-	OutputStats   = 1 << 2,
-	OutputLines   = 1 << 3,
+
+	/** Add `// REASON: foo` comments next to emitted code blocks to identify which other block uses it or other reason why it was included */
+	OutputReasons  = 1 << 1,
+
+	/** Add a comment describing how many functions, structs, etc. were emitted */
+	OutputStats    = 1 << 2,
+
+	/** Add `#line <line> <filename>` directives for each emitted code block */
+	OutputLines    = 1 << 3,
+
+	/** Keep original global scope comment lines, such as `// #define FOO 123` that were added by UE shader preprocessor */
+	OutputCommentLines = 1 << 4,
 };
 ENUM_CLASS_FLAGS(EMinifyShaderFlags)
 
