[shaders] tiny change to allow skipping output of directcompile.txt file ignoring bool in compile input (used in some forthcoming debug code which dumps extra copies of preprocessed shader code and doesn't need these)

diff --git a/Engine/Source/Developer/ShaderCompilerCommon/Private/ShaderCompilerCommon.cpp b/Engine/Source/Developer/ShaderCompilerCommon/Private/ShaderCompilerCommon.cpp
--- a/Engine/Source/Developer/ShaderCompilerCommon/Private/ShaderCompilerCommon.cpp
+++ b/Engine/Source/Developer/ShaderCompilerCommon/Private/ShaderCompilerCommon.cpp
@@ -1475,7 +1475,7 @@ namespace UE::ShaderCompilerCommon
 			DumpDebugUSF(Input, ModifiedPreprocessedSource, Options.HlslCCFlags, *BaseSourceFilename);
 			PRAGMA_ENABLE_DEPRECATION_WARNINGS
 
-			if (Input.bGenerateDirectCompileFile)
+			if (Input.bGenerateDirectCompileFile && !Options.bSkipDirectCompileTxt)
 			{
 				// note: CreateShaderCompileWorkerDirectCommandLine should be internalized once the deprecation window ends
 				// (i.e. still should remain in this file but not be declared in the header)
diff --git a/Engine/Source/Developer/ShaderCompilerCommon/Public/ShaderCompilerCommon.h b/Engine/Source/Developer/ShaderCompilerCommon/Public/ShaderCompilerCommon.h
--- a/Engine/Source/Developer/ShaderCompilerCommon/Public/ShaderCompilerCommon.h
+++ b/Engine/Source/Developer/ShaderCompilerCommon/Public/ShaderCompilerCommon.h
@@ -99,6 +99,7 @@ namespace UE::ShaderCompilerCommon
 		const TCHAR* FilenamePrefix = nullptr;
 		TFunction<FString()> AppendPreSource{};
 		TFunction<FString()> AppendPostSource{};
+		bool bSkipDirectCompileTxt = false;
 	};
 	/*
 	 * Dumps common debug information (preprocessed .usf as constructed by GetDebugShaderContents, and a directcompile.txt file 
