Moving ShaderSymbolExport implementation into its own compilation unit.

diff --git a/Engine/Source/Developer/ShaderCompilerCommon/Private/ShaderCompilerCommon.cpp b/Engine/Source/Developer/ShaderCompilerCommon/Private/ShaderCompilerCommon.cpp
--- a/Engine/Source/Developer/ShaderCompilerCommon/Private/ShaderCompilerCommon.cpp
+++ b/Engine/Source/Developer/ShaderCompilerCommon/Private/ShaderCompilerCommon.cpp
@@ -2565,73 +2565,3 @@ namespace CrossCompiler
 	}
 
 } // namespace CrossCompiler
-
-#if WITH_ENGINE
-
-const TCHAR* FShaderSymbolExport::ZipFileBaseLeafName = TEXT("ShaderSymbols");
-const TCHAR* FShaderSymbolExport::ZipFileExtension = TEXT(".zip");
-
-void FShaderSymbolExport::NotifyShaderCompilersShutdown()
-{
-	constexpr bool bMergingEnabled = false; // TODO: Add ability to read zip files and turn this on
-	if (bMultiprocessOwner && ZipWriter && bMergingEnabled)
-	{
-		IPlatformFile& PlatformFile = FPlatformFileManager::Get().GetPlatformFile();
-		TArray<FString> ZipsToMergeIn;
-		PlatformFile.FindFiles(ZipsToMergeIn, *ExportPath, ZipFileExtension);
-		ZipsToMergeIn.RemoveAll([](const FString& FileName)
-			{
-				return FPathViews::GetBaseFilename(FileName) == ZipFileBaseLeafName;
-			});
-		for (const FString& ZipFile : ZipsToMergeIn)
-		{
-			struct FZipArchiveReader
-			{
-				FZipArchiveReader(IFileHandle* InFileHandle) : FileHandle(InFileHandle) {}
-				bool IsValid() { return false; }
-				TConstArrayView<FString> GetEmbeddedFileNames() const { return TConstArrayView<FString>(); }
-				bool TryReadFile(FStringView FileName, TArray<uint8>& OutData) { return false; }
-				TUniquePtr<IFileHandle> FileHandle;
-			};
-			FZipArchiveReader Reader(PlatformFile.OpenRead(*ZipFile));
-			bool bAllValid = false;
-			if (Reader.IsValid())
-			{
-				bAllValid = true;
-				for (const FString& EmbeddedFileName : Reader.GetEmbeddedFileNames())
-				{
-					TArray<uint8> Contents;
-					if (!Reader.TryReadFile(EmbeddedFileName, Contents))
-					{
-						bAllValid = false;
-						continue;
-					}
-					ZipWriter->AddFile(EmbeddedFileName, Contents, FDateTime::Now());
-				}
-			}
-			if (!bAllValid)
-			{
-				UE_LOG(LogShaderSymbolExport, Error,
-					TEXT("Failed to read from CookWorker shader symbols output file \"%s\". Some shader symbols will be missing."),
-					*ZipFile);
-			}
-			PlatformFile.DeleteFile(*ZipFile);
-		}
-	}
-	ZipWriter.Reset();
-}
-
-void FShaderSymbolExport::DeleteExistingShaderZips(IPlatformFile& PlatformFile, const FString& Directory)
-{
-	TArray<FString> ExistingZips;
-	PlatformFile.FindFiles(ExistingZips, *Directory, ZipFileExtension);
-	for (const FString& ZipFile : ExistingZips)
-	{
-		if (FPathViews::GetPathLeaf(ZipFile).StartsWith(ZipFileBaseLeafName))
-		{
-			PlatformFile.DeleteFile(*ZipFile);
-		}
-	}
-}
-
-#endif
\ No newline at end of file
diff --git a/Engine/Source/Developer/ShaderCompilerCommon/Private/ShaderSymbolExport.cpp b/Engine/Source/Developer/ShaderCompilerCommon/Private/ShaderSymbolExport.cpp
new file mode 100644
--- /dev/null
+++ b/Engine/Source/Developer/ShaderCompilerCommon/Private/ShaderSymbolExport.cpp
@@ -0,0 +1,225 @@
+// Copyright Epic Games, Inc. All Rights Reserved.
+
+#include "ShaderSymbolExport.h"
+
+#if WITH_ENGINE
+
+#include "FileUtilities/ZipArchiveWriter.h"
+#include "GenericPlatform/GenericPlatformFile.h"
+#include "HAL/PlatformFileManager.h"
+#include "Logging/LogMacros.h"
+#include "Misc/CommandLine.h"
+#include "Misc/PathViews.h"
+#include "ShaderCompilerCore.h"
+
+DECLARE_LOG_CATEGORY_CLASS(LogShaderSymbolExport, Display, Display);
+
+static const TCHAR* ZipFileBaseLeafName = TEXT("ShaderSymbols");
+static const TCHAR* ZipFileExtension = TEXT(".zip");
+
+FShaderSymbolExport::FShaderSymbolExport(FName InShaderFormat)
+	: ShaderFormat(InShaderFormat)
+{
+}
+
+FShaderSymbolExport::~FShaderSymbolExport() = default;
+
+static void DeleteExistingShaderZips(IPlatformFile& PlatformFile, const FString& Directory)
+{
+	TArray<FString> ExistingZips;
+	PlatformFile.FindFiles(ExistingZips, *Directory, ZipFileExtension);
+
+	for (const FString& ZipFile : ExistingZips)
+	{
+		if (FPathViews::GetPathLeaf(ZipFile).StartsWith(ZipFileBaseLeafName))
+		{
+			PlatformFile.DeleteFile(*ZipFile);
+		}
+	}
+}
+
+void FShaderSymbolExport::Initialize()
+{
+	const bool bSymbolsEnabled = ShouldWriteShaderSymbols(ShaderFormat);
+	const bool bForceSymbols = FParse::Value(FCommandLine::Get(), TEXT("-ShaderSymbolsExport="), ExportPath);
+
+	if (bSymbolsEnabled || bForceSymbols)
+	{
+		// if no command line path is provided, look to the cvar first
+		if (ExportPath.IsEmpty())
+		{
+			if (GetShaderSymbolPathOverride(ExportPath, ShaderFormat))
+			{
+				ExportPath = IFileManager::Get().ConvertToAbsolutePathForExternalAppForWrite(*ExportPath);
+			}
+		}
+
+		// if there was no path set via command line or the cvar, fall back to our default
+		if (ExportPath.IsEmpty())
+		{
+			ExportPath = IFileManager::Get().ConvertToAbsolutePathForExternalAppForWrite(
+				*(FPaths::ProjectSavedDir() / TEXT("ShaderSymbols") / ShaderFormat.ToString()));
+		}
+
+		IPlatformFile& PlatformFile = FPlatformFileManager::Get().GetPlatformFile();
+		bExportShaderSymbols = PlatformFile.CreateDirectoryTree(*ExportPath);
+
+		if (!bExportShaderSymbols)
+		{
+			UE_LOG(LogShaderSymbolExport, Error, TEXT("Failed to create shader symbols output directory. Shader symbol export will be disabled."));
+		}
+		else
+		{
+			// Check if the export mode is to an uncompressed archive or loose files.
+			bool bExportAsZip = ShouldWriteShaderSymbolsAsZip(ShaderFormat);
+			if (bExportAsZip || FParse::Param(FCommandLine::Get(), TEXT("ShaderSymbolsExportZip")))
+			{
+				uint32 MultiprocessId = 0;
+				FParse::Value(FCommandLine::Get(), TEXT("-MultiprocessId="), MultiprocessId);
+				FString LeafName;
+				bMultiprocessOwner = MultiprocessId == 0;
+				if (bMultiprocessOwner)
+				{
+					DeleteExistingShaderZips(PlatformFile, ExportPath);
+					LeafName = FString::Printf(TEXT("%s%s"), ZipFileBaseLeafName, ZipFileExtension);
+				}
+				else
+				{
+					LeafName = FString::Printf(TEXT("%s_%d%s"), ZipFileBaseLeafName, MultiprocessId, ZipFileExtension);
+				}
+				FString SingleFilePath = ExportPath / LeafName;
+
+				IFileHandle* OutputZipFile = PlatformFile.OpenWrite(*SingleFilePath);
+				if (!OutputZipFile)
+				{
+					UE_LOG(LogShaderSymbolExport, Error, TEXT("Failed to create shader symbols output file \"%s\". Shader symbol export will be disabled."), *SingleFilePath);
+					bExportShaderSymbols = false;
+				}
+				else
+				{
+					ZipWriter = MakeUnique<FZipArchiveWriter>(OutputZipFile);
+				}
+			}
+		}
+	}
+
+	if (bExportShaderSymbols)
+	{
+		UE_LOG(LogShaderSymbolExport, Display, TEXT("Shader symbol export enabled. Output directory: \"%s\""), *ExportPath);
+		if (ZipWriter)
+		{
+			UE_LOG(LogShaderSymbolExport, Display, TEXT("Shader symbol zip mode enabled. Shader symbols will be archived in a single (uncompressed) zip file."));
+		}
+	}
+}
+
+void FShaderSymbolExport::WriteSymbolData(const FString& Filename, TConstArrayView<uint8> Contents)
+{
+	// Skip this symbol data if we've already exported it before.
+	bool bAlreadyInSet = false;
+	ExportedShaders.Add(Filename, &bAlreadyInSet);
+	if (bAlreadyInSet)
+	{
+		// We've already exported this shader hash
+		return;
+	}
+
+	// Emit periodic log messages detailing the size of the shader symbols output file/directory.
+	static uint64 LastReport = 0;
+	TotalSymbolDataBytes += Contents.Num();
+	TotalSymbolData++;
+
+	if ((TotalSymbolDataBytes - LastReport) >= (64 * 1024 * 1024))
+	{
+		UE_LOG(LogShaderSymbolExport, Display, TEXT("Shader symbols export size: %.2f MB, count: %llu"),
+			double(TotalSymbolDataBytes) / (1024.0 * 1024.0), TotalSymbolData);
+		LastReport = TotalSymbolDataBytes;
+	}
+
+	if (ZipWriter)
+	{
+		// Append the platform data to the zip file
+		ZipWriter->AddFile(Filename, Contents, FDateTime::Now());
+	}
+	else
+	{
+		IPlatformFile& PlatformFile = FPlatformFileManager::Get().GetPlatformFile();
+
+		// Write the symbols to the export directory
+		const FString OutputPath = ExportPath / Filename;
+		const FString Directory = FPaths::GetPath(OutputPath);
+
+		// Filename could contain extra folders, so we need to make sure they exist first.
+		if (!PlatformFile.CreateDirectoryTree(*Directory))
+		{
+			UE_LOG(LogShaderSymbolExport, Error, TEXT("Failed to create shader symbol directory \"%s\"."), *Directory);
+		}
+		else
+		{
+			IFileHandle* File = PlatformFile.OpenWrite(*OutputPath);
+			if (!File || !File->Write(Contents.GetData(), Contents.Num()))
+			{
+				UE_LOG(LogShaderSymbolExport, Error, TEXT("Failed to export shader symbols \"%s\"."), *OutputPath);
+			}
+
+			if (File)
+			{
+				delete File;
+			}
+		}
+	}
+}
+
+struct FZipArchiveReader
+{
+	FZipArchiveReader(IFileHandle* InFileHandle) : FileHandle(InFileHandle) {}
+	bool IsValid() { return false; }
+	TConstArrayView<FString> GetEmbeddedFileNames() const { return TConstArrayView<FString>(); }
+	bool TryReadFile(FStringView FileName, TArray<uint8>& OutData) { return false; }
+	TUniquePtr<IFileHandle> FileHandle;
+};
+
+void FShaderSymbolExport::NotifyShaderCompilersShutdown()
+{
+	constexpr bool bMergingEnabled = false; // TODO: Add ability to read zip files and turn this on
+	if (bMultiprocessOwner && ZipWriter && bMergingEnabled)
+	{
+		IPlatformFile& PlatformFile = FPlatformFileManager::Get().GetPlatformFile();
+		TArray<FString> ZipsToMergeIn;
+		PlatformFile.FindFiles(ZipsToMergeIn, *ExportPath, ZipFileExtension);
+		ZipsToMergeIn.RemoveAll([](const FString& FileName)
+			{
+				return FPathViews::GetBaseFilename(FileName) == ZipFileBaseLeafName;
+			});
+
+		for (const FString& ZipFile : ZipsToMergeIn)
+		{
+			FZipArchiveReader Reader(PlatformFile.OpenRead(*ZipFile));
+			bool bAllValid = false;
+			if (Reader.IsValid())
+			{
+				bAllValid = true;
+				for (const FString& EmbeddedFileName : Reader.GetEmbeddedFileNames())
+				{
+					TArray<uint8> Contents;
+					if (!Reader.TryReadFile(EmbeddedFileName, Contents))
+					{
+						bAllValid = false;
+						continue;
+					}
+					ZipWriter->AddFile(EmbeddedFileName, Contents, FDateTime::Now());
+				}
+			}
+			if (!bAllValid)
+			{
+				UE_LOG(LogShaderSymbolExport, Error,
+					TEXT("Failed to read from CookWorker shader symbols output file \"%s\". Some shader symbols will be missing."),
+					*ZipFile);
+			}
+			PlatformFile.DeleteFile(*ZipFile);
+		}
+	}
+	ZipWriter.Reset();
+}
+
+#endif // WITH_ENGINE
diff --git a/Engine/Source/Developer/ShaderCompilerCommon/Public/ShaderSymbolExport.h b/Engine/Source/Developer/ShaderCompilerCommon/Public/ShaderSymbolExport.h
--- a/Engine/Source/Developer/ShaderCompilerCommon/Public/ShaderSymbolExport.h
+++ b/Engine/Source/Developer/ShaderCompilerCommon/Public/ShaderSymbolExport.h
@@ -4,22 +4,20 @@
 
 #if WITH_ENGINE
 
+#include "CoreTypes.h"
+#include "Containers/ContainersFwd.h"
+#include "Containers/Set.h"
+#include "Templates/UniquePtr.h"
 #include "Serialization/MemoryReader.h"
-#include "HAL/FileManager.h"
-#include "HAL/IConsoleManager.h"
-#include "HAL/PlatformFileManager.h"
-#include "GenericPlatform/GenericPlatformFile.h"
-#include "Misc/CommandLine.h"
-#include "Misc/ConfigCacheIni.h"
-#include "Misc/Paths.h"
-#include "FileUtilities/ZipArchiveWriter.h"
-
-DECLARE_LOG_CATEGORY_CLASS(LogShaderSymbolExport, Display, Display);
-
-class FShaderSymbolExport
+
+class FZipArchiveWriter;
+
+class SHADERCOMPILERCOMMON_API FShaderSymbolExport
 {
 public:
+	FShaderSymbolExport() = delete;
 	FShaderSymbolExport(FName InShaderFormat);
+	~FShaderSymbolExport();
 
 	/** Should be called from IShaderFormat::NotifyShaderCompiled implementation.
 	*   Template type is the platform specific symbol data structure.
@@ -28,10 +26,11 @@ public:
 	void NotifyShaderCompiled(const TConstArrayView<uint8>& PlatformSymbolData);
 
 	/** Called at the end of a cook to free resources and finalize artifacts created during the cook. */
-	SHADERCOMPILERCOMMON_API void NotifyShaderCompilersShutdown();
+	void NotifyShaderCompilersShutdown();
 
 private:
 	void Initialize();
+	void WriteSymbolData(const FString& Filename, TConstArrayView<uint8> Contents);
 
 	const FName ShaderFormat;
 
@@ -42,97 +41,15 @@ private:
 	uint64 TotalSymbolDataBytes{ 0 };
 	uint64 TotalSymbolData{ 0 };
 	bool bExportShaderSymbols{ false };
+
 	/**
 	 * If true, the current process is the first process in a multiprocess group, or is not in a group,
 	 * and should combine artifacts produced by the other processes. Will also be false if no combination
 	 * is necessary for given settings.
 	 */
 	bool bMultiprocessOwner{ false };
-	SHADERCOMPILERCOMMON_API static void DeleteExistingShaderZips(IPlatformFile& PlatformFile, const FString& Directory);
-	SHADERCOMPILERCOMMON_API static const TCHAR* ZipFileBaseLeafName;
-	SHADERCOMPILERCOMMON_API static const TCHAR* ZipFileExtension;
 };
 
-inline FShaderSymbolExport::FShaderSymbolExport(FName InShaderFormat)
-	: ShaderFormat(InShaderFormat)
-{
-}
-
-inline void FShaderSymbolExport::Initialize()
-{
-	const bool bSymbolsEnabled = ShouldWriteShaderSymbols(ShaderFormat);
-	const bool bForceSymbols = FParse::Value(FCommandLine::Get(), TEXT("-ShaderSymbolsExport="), ExportPath);
-
-	if (bSymbolsEnabled || bForceSymbols)
-	{
-		// if no command line path is provided, look to the cvar first
-		if (ExportPath.IsEmpty())
-		{
-			if (GetShaderSymbolPathOverride(ExportPath, ShaderFormat))
-			{
-				ExportPath = IFileManager::Get().ConvertToAbsolutePathForExternalAppForWrite(*ExportPath);
-			}
-		}
-
-		// if there was no path set via command line or the cvar, fall back to our default
-		if (ExportPath.IsEmpty())
-		{
-			ExportPath = IFileManager::Get().ConvertToAbsolutePathForExternalAppForWrite(
-				*(FPaths::ProjectSavedDir() / TEXT("ShaderSymbols") / ShaderFormat.ToString()));
-		}
-
-		IPlatformFile& PlatformFile = FPlatformFileManager::Get().GetPlatformFile();
-		bExportShaderSymbols = PlatformFile.CreateDirectoryTree(*ExportPath);
-
-		if (!bExportShaderSymbols)
-		{
-			UE_LOG(LogShaderSymbolExport, Error, TEXT("Failed to create shader symbols output directory. Shader symbol export will be disabled."));
-		}
-		else
-		{
-			// Check if the export mode is to an uncompressed archive or loose files.
-			bool bExportAsZip = ShouldWriteShaderSymbolsAsZip(ShaderFormat);
-			if (bExportAsZip || FParse::Param(FCommandLine::Get(), TEXT("ShaderSymbolsExportZip")))
-			{
-				uint32 MultiprocessId = 0;
-				FParse::Value(FCommandLine::Get(), TEXT("-MultiprocessId="), MultiprocessId);
-				FString LeafName;
-				bMultiprocessOwner = MultiprocessId == 0;
-				if (bMultiprocessOwner)
-				{
-					DeleteExistingShaderZips(PlatformFile, ExportPath);
-					LeafName = FString::Printf(TEXT("%s%s"), ZipFileBaseLeafName, ZipFileExtension);
-				}
-				else
-				{
-					LeafName = FString::Printf(TEXT("%s_%d%s"), ZipFileBaseLeafName, MultiprocessId, ZipFileExtension);
-				}
-				FString SingleFilePath = ExportPath / LeafName;
-
-				IFileHandle* OutputZipFile = PlatformFile.OpenWrite(*SingleFilePath);
-				if (!OutputZipFile)
-				{
-					UE_LOG(LogShaderSymbolExport, Error, TEXT("Failed to create shader symbols output file \"%s\". Shader symbol export will be disabled."), *SingleFilePath);
-					bExportShaderSymbols = false;
-				}
-				else
-				{
-					ZipWriter = MakeUnique<FZipArchiveWriter>(OutputZipFile);
-				}
-			}
-		}
-	}
-
-	if (bExportShaderSymbols)
-	{
-		UE_LOG(LogShaderSymbolExport, Display, TEXT("Shader symbol export enabled. Output directory: \"%s\""), *ExportPath);
-		if (ZipWriter)
-		{
-			UE_LOG(LogShaderSymbolExport, Display, TEXT("Shader symbol zip mode enabled. Shader symbols will be archived in a single (uncompressed) zip file."));
-		}
-	}
-}
-
 template<typename TPlatformShaderSymbolData>
 inline void FShaderSymbolExport::NotifyShaderCompiled(const TConstArrayView<uint8>& PlatformSymbolData)
 {
@@ -158,63 +75,10 @@ inline void FShaderSymbolExport::NotifyShaderCompiled(const TConstArrayView<uint
 			{
 				const FString FileName = SymbolData.GetFilename();
 
-				// Skip this symbol data if we've already exported it before.
-				bool bAlreadyInSet = false;
-				ExportedShaders.Add(FileName, &bAlreadyInSet);
-				if (bAlreadyInSet)
-				{
-					// We've already exported this shader hash
-					continue;
-				}
-
-				// Emit periodic log messages detailing the size of the shader symbols output file/directory.
-				static uint64 LastReport = 0;
-				TotalSymbolDataBytes += Contents.Num();
-				TotalSymbolData++;
-
-				if ((TotalSymbolDataBytes - LastReport) >= (64 * 1024 * 1024))
-				{
-					UE_LOG(LogShaderSymbolExport, Display, TEXT("Shader symbols export size: %.2f MB, count: %llu"),
-						double(TotalSymbolDataBytes) / (1024.0 * 1024.0), TotalSymbolData);
-					LastReport = TotalSymbolDataBytes;
-				}
-
-				if (ZipWriter)
-				{
-					// Append the platform data to the zip file
-					ZipWriter->AddFile(FileName, Contents, FDateTime::Now());
-				}
-				else
-				{
-					IPlatformFile& PlatformFile = FPlatformFileManager::Get().GetPlatformFile();
-
-					// Write the symbols to the export directory
-					const FString OutputPath = ExportPath / FileName;
-					const FString Directory = FPaths::GetPath(OutputPath);
-
-					// FileName could contain extra folders, so we need to make sure they exist first.
-					if (!PlatformFile.CreateDirectoryTree(*Directory))
-					{
-						UE_LOG(LogShaderSymbolExport, Error, TEXT("Failed to create shader symbol directory \"%s\"."), *Directory);
-					}
-					else
-					{
-						IFileHandle* File = PlatformFile.OpenWrite(*OutputPath);
-						if (!File || !File->Write(Contents.GetData(), Contents.Num()))
-						{
-							UE_LOG(LogShaderSymbolExport, Error, TEXT("Failed to export shader symbols \"%s\"."), *OutputPath);
-						}
-
-						if (File)
-						{
-							delete File;
-						}
-					}
-				}
+				WriteSymbolData(FileName, Contents);
 			}
 		}
 	}
 }
 
-
 #endif // WITH_ENGINE
