ShaderSymbolExport: Fix missing header for non-unity compile.
#rnx

diff --git a/Engine/Source/Developer/ShaderCompilerCommon/Private/ShaderSymbolExport.cpp b/Engine/Source/Developer/ShaderCompilerCommon/Private/ShaderSymbolExport.cpp
--- a/Engine/Source/Developer/ShaderCompilerCommon/Private/ShaderSymbolExport.cpp
+++ b/Engine/Source/Developer/ShaderCompilerCommon/Private/ShaderSymbolExport.cpp
@@ -6,6 +6,7 @@
 
 #include "FileUtilities/ZipArchiveWriter.h"
 #include "GenericPlatform/GenericPlatformFile.h"
+#include "HAL/FileManager.h"
 #include "HAL/PlatformFileManager.h"
 #include "Logging/LogMacros.h"
 #include "Misc/CommandLine.h"
