ShaderSymbolExport: Rename temporary class to avoid name collision with the real class that has just been implemented.
#rnx

diff --git a/Engine/Source/Developer/ShaderCompilerCommon/Private/ShaderSymbolExport.cpp b/Engine/Source/Developer/ShaderCompilerCommon/Private/ShaderSymbolExport.cpp
--- a/Engine/Source/Developer/ShaderCompilerCommon/Private/ShaderSymbolExport.cpp
+++ b/Engine/Source/Developer/ShaderCompilerCommon/Private/ShaderSymbolExport.cpp
@@ -171,9 +171,9 @@ void FShaderSymbolExport::WriteSymbolData(const FString& Filename, TConstArrayVi
 	}
 }
 
-struct FZipArchiveReader
+struct FZipArchiveReaderTemp
 {
-	FZipArchiveReader(IFileHandle* InFileHandle) : FileHandle(InFileHandle) {}
+	FZipArchiveReaderTemp(IFileHandle* InFileHandle) : FileHandle(InFileHandle) {}
 	bool IsValid() { return false; }
 	TConstArrayView<FString> GetEmbeddedFileNames() const { return TConstArrayView<FString>(); }
 	bool TryReadFile(FStringView FileName, TArray<uint8>& OutData) { return false; }
@@ -195,7 +195,7 @@ void FShaderSymbolExport::NotifyShaderCompilersShutdown()
 
 		for (const FString& ZipFile : ZipsToMergeIn)
 		{
-			FZipArchiveReader Reader(PlatformFile.OpenRead(*ZipFile));
+			FZipArchiveReaderTemp Reader(PlatformFile.OpenRead(*ZipFile));
 			bool bAllValid = false;
 			if (Reader.IsValid())
 			{
