Fix Metal 3.0 shader compilation on iOS

diff --git a/Engine/Source/Developer/ShaderCompilerCommon/Private/ShaderConductorContext.cpp b/Engine/Source/Developer/ShaderCompilerCommon/Private/ShaderConductorContext.cpp
--- a/Engine/Source/Developer/ShaderCompilerCommon/Private/ShaderConductorContext.cpp
+++ b/Engine/Source/Developer/ShaderCompilerCommon/Private/ShaderConductorContext.cpp
@@ -203,6 +203,7 @@ namespace CrossCompiler
 	{
 		switch (Version)
 		{
+        case 30000: return "30000";
 		case 20400: return "20400";
 		case 20300: return "20300";
 		case 20200: return "20200";
