Use ZipArchiveReader to merge shader symbol zip files created during MPCook.
#rnx

diff --git a/Engine/Source/Developer/ShaderCompilerCommon/Private/ShaderSymbolExport.cpp b/Engine/Source/Developer/ShaderCompilerCommon/Private/ShaderSymbolExport.cpp
--- a/Engine/Source/Developer/ShaderCompilerCommon/Private/ShaderSymbolExport.cpp
+++ b/Engine/Source/Developer/ShaderCompilerCommon/Private/ShaderSymbolExport.cpp
@@ -4,6 +4,7 @@
 
 #if WITH_ENGINE
 
+#include "FileUtilities/ZipArchiveReader.h"
 #include "FileUtilities/ZipArchiveWriter.h"
 #include "GenericPlatform/GenericPlatformFile.h"
 #include "HAL/FileManager.h"
@@ -171,19 +172,9 @@ void FShaderSymbolExport::WriteSymbolData(const FString& Filename, TConstArrayVi
 	}
 }
 
-struct FZipArchiveReaderTemp
-{
-	FZipArchiveReaderTemp(IFileHandle* InFileHandle) : FileHandle(InFileHandle) {}
-	bool IsValid() { return false; }
-	TConstArrayView<FString> GetEmbeddedFileNames() const { return TConstArrayView<FString>(); }
-	bool TryReadFile(FStringView FileName, TArray<uint8>& OutData) { return false; }
-	TUniquePtr<IFileHandle> FileHandle;
-};
-
 void FShaderSymbolExport::NotifyShaderCompilersShutdown()
 {
-	constexpr bool bMergingEnabled = false; // TODO: Add ability to read zip files and turn this on
-	if (bMultiprocessOwner && ZipWriter && bMergingEnabled)
+	if (bMultiprocessOwner && ZipWriter)
 	{
 		IPlatformFile& PlatformFile = FPlatformFileManager::Get().GetPlatformFile();
 		TArray<FString> ZipsToMergeIn;
@@ -193,32 +184,39 @@ void FShaderSymbolExport::NotifyShaderCompilersShutdown()
 				return FPathViews::GetBaseFilename(FileName) == ZipFileBaseLeafName;
 			});
 
+#if WITH_EDITOR // FZipArchiveReader is only available in editor
 		for (const FString& ZipFile : ZipsToMergeIn)
 		{
-			FZipArchiveReaderTemp Reader(PlatformFile.OpenRead(*ZipFile));
-			bool bAllValid = false;
-			if (Reader.IsValid())
 			{
-				bAllValid = true;
-				for (const FString& EmbeddedFileName : Reader.GetEmbeddedFileNames())
+				FZipArchiveReader Reader(PlatformFile.OpenRead(*ZipFile));
+				bool bAllValid = false;
+				if (Reader.IsValid())
 				{
-					TArray<uint8> Contents;
-					if (!Reader.TryReadFile(EmbeddedFileName, Contents))
+					bAllValid = true;
+					for (const FString& EmbeddedFileName : Reader.GetFileNames())
 					{
-						bAllValid = false;
-						continue;
+						TArray<uint8> Contents;
+						if (!Reader.TryReadFile(EmbeddedFileName, Contents))
+						{
+							bAllValid = false;
+							continue;
+						}
+						ZipWriter->AddFile(EmbeddedFileName, Contents, FDateTime::Now());
 					}
-					ZipWriter->AddFile(EmbeddedFileName, Contents, FDateTime::Now());
 				}
-			}
-			if (!bAllValid)
-			{
-				UE_LOG(LogShaderSymbolExport, Error,
-					TEXT("Failed to read from CookWorker shader symbols output file \"%s\". Some shader symbols will be missing."),
-					*ZipFile);
+				if (!bAllValid)
+				{
+					UE_LOG(LogShaderSymbolExport, Error,
+						TEXT("Failed to read from CookWorker shader symbols output file \"%s\". Some shader symbols will be missing."),
+						*ZipFile);
+				}
 			}
 			PlatformFile.DeleteFile(*ZipFile);
 		}
+#else
+		UE_CLOG(!ZipsToMergeIn.IsEmpty(), LogShaderSymbolExport, Error,
+			TEXT("Cannot merge zips from multiprocess instances in %s; merging is only available in editor."), *ExportPath);
+#endif
 	}
 	ZipWriter.Reset();
 }
