ShaderMinifier - Fix handling of anonymous struct shader parameters

Declarations like `struct { uint Param; } Foo;` were previously incorrectly removed from code.

Also report explicit error when encountering 'namespace' keyword (which is currently unsupported)

diff --git a/Engine/Source/Developer/ShaderCompilerCommon/Private/ShaderMinifier.cpp b/Engine/Source/Developer/ShaderCompilerCommon/Private/ShaderMinifier.cpp
--- a/Engine/Source/Developer/ShaderCompilerCommon/Private/ShaderMinifier.cpp
+++ b/Engine/Source/Developer/ShaderCompilerCommon/Private/ShaderMinifier.cpp
@@ -4,7 +4,6 @@
 
 #include "HAL/PlatformTime.h"
 #include "Hash/CityHash.h"
-#include "HlslParser.h"
 #include "Logging/LogMacros.h"
 #include "Misc/AutomationTest.h"
 #include "String/Find.h"
@@ -386,12 +385,22 @@ static FParsedShader ParseShader(FStringView InSource, FDiagnostics& Output)
 					bHasType = true;
 				}
 
-				if (ExpressionBlockIndex > 0)
+				// If struct body is not the last block, it must be followed by a variable name
+				// i.e. `struct Foo { ... } Blah;` or `struct { ... } Blah;` or `struct Foo { ... } Blah = { expression };`
+				if (BodyBlockIndex > 0 && BodyBlockIndex + 1 < PendingBlocks.Num())
 				{
-					NameBlockIndex = ExpressionBlockIndex - 1;
+					NameBlockIndex = BodyBlockIndex + 1;
 					PendingBlocks[NameBlockIndex].Type = EBlockType::Name;
 					bHasName = true;
 				}
+
+				// If there is an expression block, we expect a named variable to also exist
+				// i.e. `struct Foo { ... } Blah = { expression };` 
+				if (ExpressionBlockIndex > 0 && NameBlockIndex == INDEX_NONE)
+				{
+					AddDiagnostic(Output.Errors, TEXT("Initialized struct variables must be named"));
+					return;
+				}
 			}
 			else if (ChunkType == ECodeChunkType::CBuffer)
 			{
@@ -581,6 +590,11 @@ static FParsedShader ParseShader(FStringView InSource, FDiagnostics& Output)
 					ChunkType = ECodeChunkType::Enum;
 					EnumBlockIndex = PendingBlocks.Num();
 				}
+				else if (Identifier == TEXT("namespace"))
+				{
+					AddDiagnostic(Output.Errors, TEXT("Unsupported keyword 'namespace'"));
+					break;
+				}
 			}
 
 			EBlockType BlockType = EBlockType::Unknown;
@@ -1650,6 +1664,14 @@ bool FShaderMinifierParserTest::RunTest(const FString& Parameters)
 		}
 	}
 
+	{
+		auto P = ParseShader(TEXT("static const struct { float4 Param; } Foo;"));
+		if (TestEqual(TEXT("ParseShader: static const anonymous struct with variable: num chunks"), P.Chunks.Num(), 1))
+		{
+			TestEqual(TEXT("ParseShader: static const anonymous struct with variable: chunk type"), P.Chunks[0].Type, ECodeChunkType::Variable);
+		}
+	}
+
 	{
 		auto P = ParseShader(TEXT("static const struct { float4 Param; } Foo = { FooCB_Param; };"));
 		if (TestEqual(TEXT("ParseShader: static const anonymous struct with variable and initializer: num chunks"), P.Chunks.Num(), 1))
@@ -1706,6 +1728,15 @@ bool FShaderMinifierParserTest::RunTest(const FString& Parameters)
 		}
 	}
 
+	{
+		FDiagnostics Diagnostics;
+		auto P = ParseShader(TEXT("ConstantBuffer<Foo> CB : register ( b123, space456);\nnamespace Foo { int A; }"), Diagnostics);
+		if (TestEqual(TEXT("ParseShader: Unsupported keyword 'namespace': num errors"), Diagnostics.Errors.Num(), 1))
+		{
+			TestEqual(TEXT("ParseShader: Unsupported keyword 'namespace'"), FStringView(Diagnostics.Errors[0].Message), FStringView(TEXT("Unsupported keyword 'namespace'")));
+		}
+	}
+
 	int32 NumErrors = ExecutionInfo.GetErrorTotal();
 
 	return NumErrors == 0;
@@ -1758,21 +1789,58 @@ float FunA()
 	return Sum(Temp);
 }
 
-float FunB()
+float FunB(int Param)
 {
-	return FunA();
+	return FunA() * (float)Param;
 }
 
 #line 1000 "MinifierTest.hlsl"
 // Test comment 1
 void EmptyFunction(){}
 
+struct
+{
+	int Foo;
+	int Bar;
+} GAnonymousStruct;
+
+struct FStructA
+{
+	int Foo;
+	int Bar;
+} GStructA;
+
+struct FStructB
+{
+	int Foo;
+} GStructB = {123};
+
+static const struct FStructC
+{
+	int Foo;
+} GStructC = { GStructA.Foo };
+
+static const struct
+{
+	int Foo;
+} GInitializedAnonymousStructA = { GStructA.Foo };
+
+static const struct
+{
+	int Foo;
+} GInitializedAnonymousStructB = { 123 };
+
+RWBuffer<float4> OutputBuffer;
+
 // Test comment 2
 [numthreads(1,1,1)]
 // Comment during function declaration
 void MainCS()
 {
-	FunB();
+	float A = FunB(GAnonymousStruct.Foo);
+	float B = FunB(GStructA.Bar + GStructB.Foo + GStructC.Foo);
+	float C = FunB(GInitializedAnonymousStructA.Foo + GInitializedAnonymousStructB.Foo);
+	OutputBuffer[0] = A + B;
 }
 )");
 
@@ -1791,8 +1859,9 @@ void MainCS()
 		return false;
 	};
 
+	FParsedShader Parsed = ParseShader(TestShaderCode);
+
 	{
-		FParsedShader Parsed = ParseShader(TestShaderCode);
 		FDiagnostics Diagnostics;
 		FString Minified = MinifyShader(Parsed, TEXT("EmptyFunction"), EMinifyShaderFlags::None, Diagnostics);
 		FParsedShader MinifiedParsed = ParseShader(Minified);
@@ -1805,7 +1874,6 @@ void MainCS()
 	}
 
 	{
-		FParsedShader Parsed = ParseShader(TestShaderCode);
 		FDiagnostics Diagnostics;
 		FString Minified = MinifyShader(Parsed, TEXT("MainCS"), EMinifyShaderFlags::OutputReasons, Diagnostics);
 		FParsedShader MinifiedParsed = ParseShader(Minified);
@@ -1816,6 +1884,13 @@ void MainCS()
 		TestTrue(TEXT("MinifyShader: MainCS: contains Sum"), ChunkPresent(MinifiedParsed, TEXT("Sum")));
 		TestTrue(TEXT("MinifyShader: MainCS: contains FunA"), ChunkPresent(MinifiedParsed, TEXT("FunA")));
 		TestTrue(TEXT("MinifyShader: MainCS: contains FunB"), ChunkPresent(MinifiedParsed, TEXT("FunB")));
+		TestTrue(TEXT("MinifyShader: MainCS: contains GAnonymousStruct"), ChunkPresent(MinifiedParsed, TEXT("GAnonymousStruct")));
+		TestTrue(TEXT("MinifyShader: MainCS: contains GStructA"), ChunkPresent(MinifiedParsed, TEXT("GStructA")));
+		TestTrue(TEXT("MinifyShader: MainCS: contains GStructB"), ChunkPresent(MinifiedParsed, TEXT("GStructB")));
+		TestTrue(TEXT("MinifyShader: MainCS: contains GStructC"), ChunkPresent(MinifiedParsed, TEXT("GStructC")));
+		TestTrue(TEXT("MinifyShader: MainCS: contains GInitializedAnonymousStructA"), ChunkPresent(MinifiedParsed, TEXT("GInitializedAnonymousStructA")));
+		TestTrue(TEXT("MinifyShader: MainCS: contains GInitializedAnonymousStructB"), ChunkPresent(MinifiedParsed, TEXT("GInitializedAnonymousStructB")));
+		TestTrue(TEXT("MinifyShader: MainCS: contains OutputBuffer"), ChunkPresent(MinifiedParsed, TEXT("OutputBuffer")));
 		TestFalse(TEXT("MinifyShader: MainCS: contains UnreferencedFunction"), ChunkPresent(MinifiedParsed, TEXT("UnreferencedFunction")));
 		TestFalse(TEXT("MinifyShader: MainCS: contains FUnreferencedStruct"), ChunkPresent(MinifiedParsed, TEXT("FUnreferencedStruct")));
 		TestFalse(TEXT("MinifyShader: MainCS: contains GUnreferencedParameter"), ChunkPresent(MinifiedParsed, TEXT("GUnreferencedParameter")));
