ShaderMinifier - Disable some of the debug features by default

diff --git a/Engine/Source/Developer/ShaderCompilerCommon/Private/ShaderCompilerCommon.cpp b/Engine/Source/Developer/ShaderCompilerCommon/Private/ShaderCompilerCommon.cpp
--- a/Engine/Source/Developer/ShaderCompilerCommon/Private/ShaderCompilerCommon.cpp
+++ b/Engine/Source/Developer/ShaderCompilerCommon/Private/ShaderCompilerCommon.cpp
@@ -691,11 +691,17 @@ bool UE::ShaderCompilerCommon::RemoveDeadCode(FString& InOutPreprocessedShaderSo
 		EntryMain += EntryIntersection;
 	}
 
-	UE::ShaderMinifier::FMinifiedShader Minified  = UE::ShaderMinifier::Minify(InOutPreprocessedShaderSource, EntryMain, 
-		UE::ShaderMinifier::EMinifyShaderFlags::OutputReasons
-		| UE::ShaderMinifier::EMinifyShaderFlags::OutputCommentLines
-		| UE::ShaderMinifier::EMinifyShaderFlags::OutputStats
-		| UE::ShaderMinifier::EMinifyShaderFlags::OutputLines);
+	UE::ShaderMinifier::EMinifyShaderFlags ExtraFlags = UE::ShaderMinifier::EMinifyShaderFlags::None;
+
+#if 0 // Extra features that may be useful during development / debugging
+	ExtraFlags |= UE::ShaderMinifier::EMinifyShaderFlags::OutputReasons // Output a comment every struct/function describing why it was included (i.e. which code block uses it)
+	           |  UE::ShaderMinifier::EMinifyShaderFlags::OutputStats;  // Output a comment detailing how many blocks of each type (functions/structs/etc.) were emitted
+#endif
+
+	UE::ShaderMinifier::FMinifiedShader Minified  = UE::ShaderMinifier::Minify(InOutPreprocessedShaderSource, EntryMain,
+		  UE::ShaderMinifier::EMinifyShaderFlags::OutputCommentLines // Preserve comments that were left after preprocessing
+		| UE::ShaderMinifier::EMinifyShaderFlags::OutputLines        // Emit #line directives
+		| ExtraFlags);
 
 	if (Minified.Success())
 	{
@@ -704,7 +710,7 @@ bool UE::ShaderCompilerCommon::RemoveDeadCode(FString& InOutPreprocessedShaderSo
 	}
 	else
 	{
-		OutErrors.Add(TEXT("Shader minification failed."));
+		OutErrors.Add(TEXT("warning: Shader minification failed."));
 		return false;
 	}
 }
