ShaderMinifier - Add basic namespace support

* Track namespaces for symbols
* Re-emit namespace scopes
* Conservatively include symbols from all namespaces in minified code (as if they were all global)
* Additional parsing work will be required to more accurately track used symbols based on active namespaces (i.e. handle `using` keyword, etc.)
* Added RemoveDeadCode() overload that takes an explicit list of required symbols instead of just an entry point name

diff --git a/Engine/Source/Developer/ShaderCompilerCommon/Private/ShaderCompilerCommon.cpp b/Engine/Source/Developer/ShaderCompilerCommon/Private/ShaderCompilerCommon.cpp
--- a/Engine/Source/Developer/ShaderCompilerCommon/Private/ShaderCompilerCommon.cpp
+++ b/Engine/Source/Developer/ShaderCompilerCommon/Private/ShaderCompilerCommon.cpp
@@ -672,25 +672,8 @@ void UE::ShaderCompilerCommon::ParseRayTracingEntryPoint(const FString& Input, F
 	}
 }
 
-bool UE::ShaderCompilerCommon::RemoveDeadCode(FString& InOutPreprocessedShaderSource, const FString& EntryPoint, TArray<FShaderCompilerError>& OutErrors)
+bool UE::ShaderCompilerCommon::RemoveDeadCode(FString& InOutPreprocessedShaderSource, TConstArrayView<FStringView> RequiredSymbols, TArray<FShaderCompilerError>& OutErrors)
 {
-	FString EntryMain;
-	FString EntryAnyHit;
-	FString EntryIntersection;
-	UE::ShaderCompilerCommon::ParseRayTracingEntryPoint(EntryPoint, EntryMain, EntryAnyHit, EntryIntersection);
-
-	if (!EntryAnyHit.IsEmpty())
-	{
-		EntryMain += TEXT(";");
-		EntryMain += EntryAnyHit;
-	}
-
-	if (!EntryIntersection.IsEmpty())
-	{
-		EntryMain += TEXT(";");
-		EntryMain += EntryIntersection;
-	}
-
 	UE::ShaderMinifier::EMinifyShaderFlags ExtraFlags = UE::ShaderMinifier::EMinifyShaderFlags::None;
 
 #if 0 // Extra features that may be useful during development / debugging
@@ -698,7 +681,7 @@ bool UE::ShaderCompilerCommon::RemoveDeadCode(FString& InOutPreprocessedShaderSo
 	           |  UE::ShaderMinifier::EMinifyShaderFlags::OutputStats;  // Output a comment detailing how many blocks of each type (functions/structs/etc.) were emitted
 #endif
 
-	UE::ShaderMinifier::FMinifiedShader Minified  = UE::ShaderMinifier::Minify(InOutPreprocessedShaderSource, EntryMain,
+	UE::ShaderMinifier::FMinifiedShader Minified  = UE::ShaderMinifier::Minify(InOutPreprocessedShaderSource, RequiredSymbols,
 		  UE::ShaderMinifier::EMinifyShaderFlags::OutputCommentLines // Preserve comments that were left after preprocessing
 		| UE::ShaderMinifier::EMinifyShaderFlags::OutputLines        // Emit #line directives
 		| ExtraFlags);
@@ -715,6 +698,30 @@ bool UE::ShaderCompilerCommon::RemoveDeadCode(FString& InOutPreprocessedShaderSo
 	}
 }
 
+bool UE::ShaderCompilerCommon::RemoveDeadCode(FString& InOutPreprocessedShaderSource, const FString& EntryPoint, TArray<FShaderCompilerError>& OutErrors)
+{
+	TArray<FStringView> RequiredSymbols;
+
+	FString EntryMain;
+	FString EntryAnyHit;
+	FString EntryIntersection;
+	UE::ShaderCompilerCommon::ParseRayTracingEntryPoint(EntryPoint, EntryMain, EntryAnyHit, EntryIntersection);
+
+	RequiredSymbols.Add(EntryMain);
+
+	if (!EntryAnyHit.IsEmpty())
+	{
+		RequiredSymbols.Add(EntryAnyHit);
+	}
+
+	if (!EntryIntersection.IsEmpty())
+	{
+		RequiredSymbols.Add(EntryIntersection);
+	}
+
+	return UE::ShaderCompilerCommon::RemoveDeadCode(InOutPreprocessedShaderSource, RequiredSymbols, OutErrors);
+}
+
 void HandleReflectedGlobalConstantBufferMember(
 	const FString& InMemberName,
 	uint32 ConstantBufferIndex,
diff --git a/Engine/Source/Developer/ShaderCompilerCommon/Private/ShaderMinifier.cpp b/Engine/Source/Developer/ShaderCompilerCommon/Private/ShaderMinifier.cpp
--- a/Engine/Source/Developer/ShaderCompilerCommon/Private/ShaderMinifier.cpp
+++ b/Engine/Source/Developer/ShaderCompilerCommon/Private/ShaderMinifier.cpp
@@ -12,7 +12,6 @@
 DEFINE_LOG_CATEGORY_STATIC(LogShaderMinifier, Log, All);
 
 // TODO:
-// - track namespaces
 // - preserve multi-line #define
 
 namespace UE::ShaderMinifier
@@ -166,23 +165,6 @@ static FStringView SkipUntilStr(FStringView Haystack, FStringView Needle)
 	return SkipUntil(Haystack, [Needle](FStringView  S) { return S.StartsWith(Needle, ESearchCase::CaseSensitive); });
 }
 
-static int32 FindFirstOf(FStringView Haystack, FStringView Needle)
-{
-	int32 Len = Haystack.Len();
-	for (int32 i = 0; i < Len; ++i)
-	{
-		TCHAR C = Haystack[i];
-		for (TCHAR C2 : Needle)
-		{
-			if (C == C2)
-			{
-				return i;
-			}
-		}
-	}
-	return INDEX_NONE;
-}
-
 static FStringView ExtractBlock(FStringView Source, TCHAR DelimBegin, TCHAR DelimEnd)
 {
 	// TODO: handle comments
@@ -226,7 +208,7 @@ static FStringView ExtractBlock(FStringView Source, TCHAR DelimBegin, TCHAR Deli
 }
 
 enum class EBlockType : uint8 {
-	Unknown,
+	Unknown,	// various identifiers and keywords that we did not need to or could not identify
 	Keyword,	// e.g. struct, switch, register
 	Attribute,	// e.g. `[numthreads(8,8,1)]`
 	Type,		// return type of function or struct/cbuffer/variable type
@@ -240,6 +222,7 @@ enum class EBlockType : uint8 {
 	Expression,
 	Directive,  // #define, #pragma, #line, etc.
 	NamespaceDelimiter, // e.g. :: in an identifier like Foo::bar
+	PtrOrRef, // e.g. '*' or '&' as part of the type
 };
 
 struct FCodeBlock
@@ -258,6 +241,29 @@ enum class ECodeChunkType {
 	Define,
 	Pragma,
 	CommentLine, // Single line comment
+	Namespace,
+	Using,
+};
+
+struct FNamespace
+{
+	FNamespace() = default;
+	FNamespace(TConstArrayView<FStringView> InStack)
+	{
+		if (!InStack.IsEmpty())
+		{
+			for (const FStringView& Part : InStack)
+			{
+				FullName += Part;
+				FullName += TEXT("::");
+			}
+			FullName.LeftChopInline(2);
+		}
+		Stack = InStack;
+	}
+
+	FString FullName; // i.e. Foo::Bar::Baz
+	TArray<FStringView> Stack; // i.e. [Foo, Bar, Baz]
 };
 
 struct FCodeChunk
@@ -265,6 +271,8 @@ struct FCodeChunk
 	ECodeChunkType Type = ECodeChunkType::Unknown;
 	TArray<FCodeBlock> Blocks;
 
+	int32 Namespace = INDEX_NONE; // Unique namespace ID (INDEX_NONE = global)
+
 	// Indicates whether the code for this chunk can be used as-is.
 	// One example where we have to do custom code emission is when a named struct and a variable are declared in one chunk.
 	// The struct type may be referenced, but the variable may be removed. In this case we have to emit the type declaration only.
@@ -303,10 +311,59 @@ struct FCodeChunk
 struct FParsedShader
 {
 	FStringView Source;
-
 	TArray<FCodeChunk> Chunks;
+	TArray<FNamespace> Namespaces;
+};
+
+struct FNamespaceTracker
+{
+	TMap<FString, int32> UniqueNamespaceMap;
+	TArray<FNamespace> UniqueNamespaceArray;
+	TArray<FStringView> NamespaceStack;
+	TArray<int32> NamespaceIdStack;
+
+	FNamespaceTracker() = default;
+
+	void Push(FStringView Name)
+	{
+		NamespaceStack.Push(Name);
+		FNamespace NamespaceEntry(NamespaceStack);
+		int32& EntryIndex = UniqueNamespaceMap.FindOrAdd(NamespaceEntry.FullName, INDEX_NONE);
+		if (EntryIndex == INDEX_NONE)
+		{
+			EntryIndex = UniqueNamespaceArray.Num();
+			UniqueNamespaceArray.Add(MoveTemp(NamespaceEntry));
+		}
+		NamespaceIdStack.Push(EntryIndex);
+	}
+
+	bool Pop()
+	{
+		if (NamespaceStack.IsEmpty())
+		{
+			return false;
+		}
+		else
+		{
+			NamespaceStack.Pop();
+			NamespaceIdStack.Pop();
+			return true;
+		}
+	}
+
+	int32 CurrentId() const
+	{
+		return NamespaceIdStack.IsEmpty() ? INDEX_NONE : NamespaceIdStack.Last();
+	}
 };
 
+FStringView ExtractNextIdentifier(FStringView Source)
+{
+	FStringView Remainder = SkipUntilNonIdentifierCharacter(Source);
+	FStringView Identifier = SubStrView(Source, 0, Source.Len() - Remainder.Len());
+	return Identifier;
+}
+
 static FParsedShader ParseShader(FStringView InSource, FDiagnostics& Output)
 {
 	FParsedShader Result;
@@ -330,6 +387,9 @@ static FParsedShader ParseShader(FStringView InSource, FDiagnostics& Output)
 	int32 BodyBlockIndex = INDEX_NONE;
 	int32 ExpressionBlockIndex = INDEX_NONE;
 
+	FNamespaceTracker NamespaceTracker;
+	FStringView		  PendingNamespace;
+
 	auto AddDiagnostic = [InSource, &Source](TArray<FDiagnosticMessage>& Output, FStringView Message)
 	{
 		FDiagnosticMessage Diagnostic;
@@ -350,7 +410,8 @@ static FParsedShader ParseShader(FStringView InSource, FDiagnostics& Output)
 		PendingBlocks.Push(NewBlock);
 	};
 
-	auto FinalizeChunk = [&]() {
+	auto FinalizeChunk = [&]() 
+	{
 		const bool bFoundArgs = ArgsBlockIndex >= 0;
 
 		bool bHasType = false;
@@ -443,6 +504,8 @@ static FParsedShader ParseShader(FStringView InSource, FDiagnostics& Output)
 				ChunkType = ECodeChunkType::Variable;
 			}
 
+			const int32 Namespace = NamespaceTracker.CurrentId();
+
 			if (ChunkType == ECodeChunkType::Struct && bHasName && bHasType)
 			{
 				// Handle simultaneous struct type and variable declaration
@@ -469,6 +532,9 @@ static FParsedShader ParseShader(FStringView InSource, FDiagnostics& Output)
 					VarChunk.Blocks.Push(PendingBlocks[i]);
 				}
 
+				StructChunk.Namespace = Namespace;
+				VarChunk.Namespace = Namespace;
+
 				Chunks.Push(StructChunk);
 				Chunks.Push(VarChunk);
 			}
@@ -476,6 +542,7 @@ static FParsedShader ParseShader(FStringView InSource, FDiagnostics& Output)
 			{
 				FCodeChunk Chunk;
 				Chunk.Type = ChunkType;
+				Chunk.Namespace = Namespace;
 				Swap(Chunk.Blocks, PendingBlocks);
 				Chunks.Push(Chunk);
 			}
@@ -566,9 +633,44 @@ static FParsedShader ParseShader(FStringView InSource, FDiagnostics& Output)
 			}
 			continue;
 		}
+		else if (PendingBlocks.IsEmpty() && Source.StartsWith(TEXT("{")))
+		{
+			if (ChunkType == ECodeChunkType::Namespace)
+			{
+				if (PendingNamespace.IsEmpty())
+				{
+					AddDiagnostic(Output.Errors, TEXT("HLSL does not support anonymous namespaces"));
+					break;
+				}
+				else
+				{
+					NamespaceTracker.Push(PendingNamespace);
+					ChunkType = ECodeChunkType::Unknown;
+					PendingNamespace = {};
+					Source = Source.Mid(1);
+				}
+				continue;
+			}
+			else
+			{
+				AddDiagnostic(Output.Errors, TEXT("Expected token '{'"));
+			}
+			continue;
+		}
+		else if (PendingBlocks.IsEmpty() && Source.StartsWith(TEXT("}")))
+		{
+			if (NamespaceTracker.Pop())
+			{
+				Source = Source.Mid(1);
+				continue;
+			}
+			else
+			{
+				AddDiagnostic(Output.Errors, TEXT("Expected token '}'"));
+			}
+		}
 
 		FStringView Remainder = SkipUntilNonIdentifierCharacter(Source);
-
 		FStringView Identifier = SubStrView(Source, 0, Source.Len() - Remainder.Len());
 
 		if (Identifier.Len())
@@ -592,9 +694,23 @@ static FParsedShader ParseShader(FStringView InSource, FDiagnostics& Output)
 				}
 				else if (Identifier == TEXT("namespace"))
 				{
-					AddDiagnostic(Output.Errors, TEXT("Unsupported keyword 'namespace'"));
-					break;
+					ChunkType = ECodeChunkType::Namespace;
+					Source = Remainder;
+					continue;
 				}
+				else if (Identifier == TEXT("using"))
+				{
+					ChunkType = ECodeChunkType::Using;
+					Source = Remainder;
+					AddBlock(EBlockType::Keyword, Identifier);
+					continue;
+				}
+			}
+			else if (ChunkType == ECodeChunkType::Namespace)
+			{
+				PendingNamespace = Identifier;
+				Source = Remainder;
+				continue;
 			}
 
 			EBlockType BlockType = EBlockType::Unknown;
@@ -643,26 +759,28 @@ static FParsedShader ParseShader(FStringView InSource, FDiagnostics& Output)
 		}
 		else if (C == '=')
 		{
-			int32 Pos = FindFirstOf(Source, TEXT("{;"));
-
-			if (Pos == INDEX_NONE)
-			{
-				AddDiagnostic(Output.Errors, TEXT("Expected block body or semicolon after '='"));
-				break;
-			}
-
 			bFoundAssignment = true;
 
-			char C2 = Source[Pos];
+			Source = SkipSpace(Source.Mid(1));
+
+			char C2 = Source[0];
 
 			if (C2 == '{')
 			{
-				Source = SubStrView(Source, Pos);
+				// extract block on the next loop iteration
 				continue;
 			}
-			else if (C2 == ';')
+			else
 			{
-				Block = SubStrView(Source, 1, Pos - 1);
+				int32 Pos = INDEX_NONE;
+
+				if (!Source.FindChar(TCHAR(';'), Pos))
+				{
+					AddDiagnostic(Output.Errors, TEXT("Expected semicolon after assignment expression"));
+					break;
+				}
+
+				Block = SubStrView(Source, 0, Pos);
 				Block = TrimSpace(Block);
 
 				int32 BlockOffset = int32(Block.GetData() - Source.GetData());
@@ -734,9 +852,18 @@ static FParsedShader ParseShader(FStringView InSource, FDiagnostics& Output)
 			Source = SubStrView(Source, 1);
 			continue;
 		}
+		else if ((C == '*' || C == '&') && !PendingBlocks.IsEmpty()) // Part of a pointer or reference declaration
+		{
+			Block = SubStrView(Source, 0, 1);
+			Source = SubStrView(Source, 1);
+
+			AddBlock(EBlockType::PtrOrRef, Block);
+
+			continue;
+		}
 		else
 		{
-			AddDiagnostic(Output.Errors, TEXT("Unexpected character"));
+			AddDiagnostic(Output.Errors, FString::Printf(TEXT("Unexpected character '%c'"), C));
 			break;
 		}
 
@@ -766,7 +893,8 @@ static FParsedShader ParseShader(FStringView InSource, FDiagnostics& Output)
 		}
 	}
 
-	std::swap(Result.Chunks, Chunks);
+	Swap(Result.Chunks, Chunks);
+	Swap(Result.Namespaces, NamespaceTracker.UniqueNamespaceArray);
 
 	return Result;
 }
@@ -938,7 +1066,7 @@ static void OutputChunk(const FCodeChunk& Chunk, FStringBuilderBase& OutputStrea
 		OutputStream << ";";
 	}
 
-	OutputStream << "\n\n";
+	OutputStream << "\n";
 }
 
 struct FCasedStringViewKeyFuncs : public DefaultKeyFuncs<FStringView>
@@ -1051,7 +1179,25 @@ static bool ParseLineDirective(FStringView Input, int32& OutLineNumber, FStringV
 	return true;
 }
 
-static FString MinifyShader(const FParsedShader& Parsed, FStringView SemicolonSeparatedEntryPoints, EMinifyShaderFlags Flags, FDiagnostics& Diagnostics)
+static void OpenNamespace(FStringBuilderBase& OutputStream, const FNamespace& Namespace)
+{
+	for (const FStringView& Name : Namespace.Stack)
+	{
+		OutputStream << TEXT("namespace ") << Name << TEXT(" { ");
+	}
+}
+
+static void CloseNamespace(FStringBuilderBase& OutputStream, const FNamespace& Namespace)
+{
+	for (const FStringView& Name : Namespace.Stack)
+	{
+		OutputStream << TEXT("}");
+	}
+
+	OutputStream << TEXT(" // namespace ") << Namespace.FullName;
+}
+
+static FString MinifyShader(const FParsedShader& Parsed, TConstArrayView<FStringView> RequiredSymbols, EMinifyShaderFlags Flags, FDiagnostics& Diagnostics)
 {
 	FStringBuilderBase OutputStream;
 
@@ -1062,9 +1208,7 @@ static FString MinifyShader(const FParsedShader& Parsed, FStringView SemicolonSe
 
 	TArray<const FCodeChunk*> PendingChunks;
 
-	TArray<FStringView> EntryPoints = SplitByChar(SemicolonSeparatedEntryPoints, ';');
-
-	for (FStringView Entry : EntryPoints)
+	for (FStringView Entry : RequiredSymbols)
 	{
 		RelevantIdentifiers.Add(Entry);
 		ProcessedIdentifiers.Add(Entry);
@@ -1351,6 +1495,8 @@ static FString MinifyShader(const FParsedShader& Parsed, FStringView SemicolonSe
 		BuildLineBreakMap(Parsed.Source, LineBreakMap, LineDirectives);
 	}
 
+	const FNamespace* CurrentNamespace = nullptr;
+
 	for (const FCodeChunk& Chunk : Parsed.Chunks)
 	{
 		auto ShouldSkipChunk = [&RelevantChunks, &Chunk, Flags]()
@@ -1367,6 +1513,12 @@ static FString MinifyShader(const FParsedShader& Parsed, FStringView SemicolonSe
 				return false;
 			}
 
+			// Always include `using` statements if they are present in the global scope
+			if (Chunk.Type == ECodeChunkType::Using)
+			{
+				return false;
+			}
+
 			if (RelevantChunks.Find(&Chunk))
 			{
 				return false;
@@ -1380,13 +1532,32 @@ static FString MinifyShader(const FParsedShader& Parsed, FStringView SemicolonSe
 			continue;
 		}
 
+		const FNamespace* PendingNamespace = Chunk.Namespace != INDEX_NONE ? &Parsed.Namespaces[Chunk.Namespace] : nullptr;
+
+		if (PendingNamespace != CurrentNamespace)
+		{
+			if (CurrentNamespace)
+			{
+				CloseNamespace(OutputStream, *CurrentNamespace);
+				OutputStream << TEXT("\n\n");
+			}
+
+			if (PendingNamespace)
+			{
+				OpenNamespace(OutputStream, *PendingNamespace);
+				OutputStream << TEXT("\n\n");
+			}
+
+			CurrentNamespace = PendingNamespace;
+		}
+
 		if (EnumHasAnyFlags(Flags, EMinifyShaderFlags::OutputReasons))
 		{
 			auto RequestedBy = ChunkRequestedBy.Find(&Chunk);
 			if (RequestedBy != nullptr)
 			{
 				const FCodeChunk* RequestedByChunk = *RequestedBy;
-				FStringView  RequestedByName  = RequestedByChunk->FindFirstBlockByType(EBlockType::Name);
+				FStringView  RequestedByName = RequestedByChunk->FindFirstBlockByType(EBlockType::Name);
 				if (!RequestedByName.IsEmpty())
 				{
 					OutputStream << TEXT("// REASON: ") << RequestedByName << TEXT("\n");
@@ -1432,6 +1603,15 @@ static FString MinifyShader(const FParsedShader& Parsed, FStringView SemicolonSe
 		}
 
 		OutputChunk(Chunk, OutputStream);
+
+		OutputStream << "\n";
+	}
+
+	if (CurrentNamespace)
+	{
+		CloseNamespace(OutputStream, *CurrentNamespace);
+		OutputStream << TEXT("\n");
+		CurrentNamespace = nullptr;
 	}
 
 	FString Output = FString(OutputStream.ToView());
@@ -1439,7 +1619,13 @@ static FString MinifyShader(const FParsedShader& Parsed, FStringView SemicolonSe
 	return Output;
 }
 
-FMinifiedShader Minify(const FStringView PreprocessedShader, const FStringView EntryPoint, EMinifyShaderFlags Flags)
+static FString MinifyShader(const FParsedShader& Parsed, FStringView EntryPoint, EMinifyShaderFlags Flags, FDiagnostics& Diagnostics)
+{
+	TArray<FStringView> RequiredSymbols = SplitByChar(EntryPoint, ';');
+	return MinifyShader(Parsed, RequiredSymbols, Flags, Diagnostics);
+}
+
+FMinifiedShader Minify(const FStringView PreprocessedShader, TConstArrayView<FStringView> RequiredSymbols, EMinifyShaderFlags Flags)
 {
 	FMinifiedShader Result;
 
@@ -1447,12 +1633,17 @@ FMinifiedShader Minify(const FStringView PreprocessedShader, const FStringView E
 
 	if (!Parsed.Chunks.IsEmpty())
 	{
-		Result.Code = MinifyShader(Parsed, EntryPoint, Flags, Result.Diagnostics);
+		Result.Code = MinifyShader(Parsed, RequiredSymbols, Flags, Result.Diagnostics);
 	}
 
 	return Result;
 }
 
+FMinifiedShader Minify(const FStringView PreprocessedShader, const FStringView EntryPoint, EMinifyShaderFlags Flags)
+{
+	return Minify(PreprocessedShader, MakeArrayView(&EntryPoint, 1), Flags);
+}
+
 } // namespace UE::ShaderMinifier
 
 #if WITH_AUTOMATION_TESTS
@@ -1729,11 +1920,12 @@ bool FShaderMinifierParserTest::RunTest(const FString& Parameters)
 	}
 
 	{
-		FDiagnostics Diagnostics;
-		auto P = ParseShader(TEXT("ConstantBuffer<Foo> CB : register ( b123, space456);\nnamespace Foo { int A; }"), Diagnostics);
-		if (TestEqual(TEXT("ParseShader: Unsupported keyword 'namespace': num errors"), Diagnostics.Errors.Num(), 1))
+		auto P = ParseShader(TEXT("namespace NS1 { void Fun() {}; } namespace NS2 { void Fun() {}; }"));
+		if (TestEqual(TEXT("ParseShader: namespaces: num chunks"), P.Chunks.Num(), 2)
+			&& TestEqual(TEXT("ParseShader: namespaces: num namespaces"), P.Namespaces.Num(), 2))
 		{
-			TestEqual(TEXT("ParseShader: Unsupported keyword 'namespace'"), FStringView(Diagnostics.Errors[0].Message), FStringView(TEXT("Unsupported keyword 'namespace'")));
+			TestEqual(TEXT("ParseShader: namespaces: chunk 0 namespace"), P.Chunks[0].Namespace, 0);
+			TestEqual(TEXT("ParseShader: namespaces: chunk 1 namespace"), P.Chunks[1].Namespace, 1);
 		}
 	}
 
@@ -1810,7 +2002,9 @@ struct FStructA
 	int Bar;
 } GStructA;
 
-struct FStructB
+namespace NS1 {
+namespace NS2 {
+static const struct FStructB
 {
 	int Foo;
 } GStructB = {123};
@@ -1819,11 +2013,14 @@ static const struct FStructC
 {
 	int Foo;
 } GStructC = { GStructA.Foo };
+}} // NS1::NS2
 
+namespace NS3 {
 static const struct
 {
 	int Foo;
 } GInitializedAnonymousStructA = { GStructA.Foo };
+} // NS3
 
 static const struct
 {
@@ -1837,6 +2034,8 @@ RWBuffer<float4> OutputBuffer;
 // Comment during function declaration
 void MainCS()
 {
+	using namespace NS1::NS2;
+	using namespace NS3;
 	float A = FunB(GAnonymousStruct.Foo);
 	float B = FunB(GStructA.Bar + GStructB.Foo + GStructC.Foo);
 	float C = FunB(GInitializedAnonymousStructA.Foo + GInitializedAnonymousStructB.Foo);
@@ -1878,6 +2077,7 @@ void MainCS()
 		FString Minified = MinifyShader(Parsed, TEXT("MainCS"), EMinifyShaderFlags::OutputReasons, Diagnostics);
 		FParsedShader MinifiedParsed = ParseShader(Minified);
 
+		// Expect true:
 		TestTrue(TEXT("MinifyShader: MainCS: contains MainCS"), ChunkPresent(MinifiedParsed, TEXT("MainCS")));
 		TestTrue(TEXT("MinifyShader: MainCS: contains FFoo"), ChunkPresent(MinifiedParsed, TEXT("FFoo")));
 		TestTrue(TEXT("MinifyShader: MainCS: contains FBar"), ChunkPresent(MinifiedParsed, TEXT("FBar")));
@@ -1891,6 +2091,8 @@ void MainCS()
 		TestTrue(TEXT("MinifyShader: MainCS: contains GInitializedAnonymousStructA"), ChunkPresent(MinifiedParsed, TEXT("GInitializedAnonymousStructA")));
 		TestTrue(TEXT("MinifyShader: MainCS: contains GInitializedAnonymousStructB"), ChunkPresent(MinifiedParsed, TEXT("GInitializedAnonymousStructB")));
 		TestTrue(TEXT("MinifyShader: MainCS: contains OutputBuffer"), ChunkPresent(MinifiedParsed, TEXT("OutputBuffer")));
+
+		// Expect false:
 		TestFalse(TEXT("MinifyShader: MainCS: contains UnreferencedFunction"), ChunkPresent(MinifiedParsed, TEXT("UnreferencedFunction")));
 		TestFalse(TEXT("MinifyShader: MainCS: contains FUnreferencedStruct"), ChunkPresent(MinifiedParsed, TEXT("FUnreferencedStruct")));
 		TestFalse(TEXT("MinifyShader: MainCS: contains GUnreferencedParameter"), ChunkPresent(MinifiedParsed, TEXT("GUnreferencedParameter")));
diff --git a/Engine/Source/Developer/ShaderCompilerCommon/Public/ShaderCompilerCommon.h b/Engine/Source/Developer/ShaderCompilerCommon/Public/ShaderCompilerCommon.h
--- a/Engine/Source/Developer/ShaderCompilerCommon/Public/ShaderCompilerCommon.h
+++ b/Engine/Source/Developer/ShaderCompilerCommon/Public/ShaderCompilerCommon.h
@@ -85,12 +85,13 @@ namespace UE::ShaderCompilerCommon
 	extern SHADERCOMPILERCOMMON_API void ParseRayTracingEntryPoint(const FString& Input, FString& OutMain, FString& OutAnyHit, FString& OutIntersection);
 
 	/**
-	* Rewrites a fully preprocessed shader source code, removing any functions or structs that are not reachable from a given entry point.
+	* Rewrites a fully preprocessed shader source code, removing any functions or structs that are not reachable from a given entry point or list of symbols.
 	* This is a high-level wrapper for UE::ShaderMinifier that should be used in all shader compiler back-ends for consistency.
 	* Any errors encountered during parsing or minification are added to OutErrors.
 	* InOutPreprocessedShaderSource is replaced with the rewritten code on success and is kept intact on failure.
 	*/
 	extern SHADERCOMPILERCOMMON_API bool RemoveDeadCode(FString& InOutPreprocessedShaderSource, const FString& EntryPoint, TArray<FShaderCompilerError>& OutErrors);
+	extern SHADERCOMPILERCOMMON_API bool RemoveDeadCode(FString& InOutPreprocessedShaderSource, TConstArrayView<FStringView> RequiredSymbols, TArray<FShaderCompilerError>& OutErrors);
 
 	struct FDebugShaderDataOptions
 	{
diff --git a/Engine/Source/Developer/ShaderCompilerCommon/Public/ShaderMinifier.h b/Engine/Source/Developer/ShaderCompilerCommon/Public/ShaderMinifier.h
--- a/Engine/Source/Developer/ShaderCompilerCommon/Public/ShaderMinifier.h
+++ b/Engine/Source/Developer/ShaderCompilerCommon/Public/ShaderMinifier.h
@@ -55,6 +55,7 @@ struct FMinifiedShader
 };
 
 extern SHADERCOMPILERCOMMON_API FMinifiedShader Minify(const FStringView PreprocessedShader, const FStringView EntryPoint, EMinifyShaderFlags Flags = EMinifyShaderFlags::None);
+extern SHADERCOMPILERCOMMON_API FMinifiedShader Minify(const FStringView PreprocessedShader, TConstArrayView<FStringView> RequiredSymbols, EMinifyShaderFlags Flags = EMinifyShaderFlags::None);
 
 } // UE::ShaderMinifier
 
