Automatic bindless resource declaration conversion step will now explicitly expand AUTO_BINDLESS_SAMPLER/RESOURCE_INDEX/VARIABLE macros

This is required because subsequent shader compilation steps assume that the code has been fully preprocessed. In particular, the dead code elimination system does not handle macros such as ones in BindlessResources.ush.

diff --git a/Engine/Shaders/Public/BindlessResources.ush b/Engine/Shaders/Public/BindlessResources.ush
--- a/Engine/Shaders/Public/BindlessResources.ush
+++ b/Engine/Shaders/Public/BindlessResources.ush
@@ -33,8 +33,9 @@
 		DEFINE_RESOURCE_VARIABLE(Type, Name); \
 		static const Type Name = GET_RESOURCE(Name);
 
-	COMPILER_DEFINE AUTO_BINDLESS_RESOURCE_INDEX(Name)          uint BindlessResource_##Name
-	COMPILER_DEFINE AUTO_BINDLESS_RESOURCE_VARIABLE(Type, Name) Type GetBindlessResource##Name() { return GetResourceFromHeap(Type, BindlessResource_##Name); }; static const Type Name = GetBindlessResource##Name()
+	// NOTE: These macros equivalent to FShaderParameterParser::ApplyBindlessModifications and must be kept in sync
+	#define AUTO_BINDLESS_RESOURCE_INDEX(Name)          uint BindlessResource_##Name
+	#define AUTO_BINDLESS_RESOURCE_VARIABLE(Type, Name) Type GetBindlessResource##Name() { return GetResourceFromHeap(Type, BindlessResource_##Name); }; static const Type Name = GetBindlessResource##Name()
 
 	#define UB_RESOURCE_MEMBER_RESOURCE(Type, Name) AUTO_BINDLESS_RESOURCE_VARIABLE(Type, Name)
 #else
@@ -59,8 +60,9 @@
 		DEFINE_SAMPLER_VARIABLE(Type, Name); \
 		static const Type Name = GET_SAMPLER(Name);
 
-	COMPILER_DEFINE AUTO_BINDLESS_SAMPLER_INDEX(Name)          uint BindlessSampler_##Name
-	COMPILER_DEFINE AUTO_BINDLESS_SAMPLER_VARIABLE(Type, Name) Type GetBindlessSampler##Name() { return GetSamplerFromHeap(Type, BindlessSampler_##Name); }; static const Type Name = GetBindlessSampler##Name()
+	// NOTE: These macros equivalent to FShaderParameterParser::ApplyBindlessModifications and must be kept in sync
+	#define AUTO_BINDLESS_SAMPLER_INDEX(Name)          uint BindlessSampler_##Name
+	#define AUTO_BINDLESS_SAMPLER_VARIABLE(Type, Name) Type GetBindlessSampler##Name() { return GetSamplerFromHeap(Type, BindlessSampler_##Name); }; static const Type Name = GetBindlessSampler##Name()
 
 	#define UB_RESOURCE_MEMBER_SAMPLER(Type, Name) AUTO_BINDLESS_SAMPLER_VARIABLE(Type, Name)
 #else
diff --git a/Engine/Shaders/Public/ShaderVersion.ush b/Engine/Shaders/Public/ShaderVersion.ush
--- a/Engine/Shaders/Public/ShaderVersion.ush
+++ b/Engine/Shaders/Public/ShaderVersion.ush
@@ -2,4 +2,4 @@
 // Each time the console command is executed it generates a new GUID. As this file is included
 // in Platform.ush (which should be included in any shader) it allows to invalidate the shader DDC.
 // 
-// GUID = 71F1E87B-7930-4F25-9575-BCBEB3A9BB59
\ No newline at end of file
+// GUID = A9EC3819-3FCC-41DC-88FB-A1AC66A46D20
diff --git a/Engine/Source/Developer/ShaderCompilerCommon/Private/ShaderParameterParser.cpp b/Engine/Source/Developer/ShaderCompilerCommon/Private/ShaderParameterParser.cpp
--- a/Engine/Source/Developer/ShaderCompilerCommon/Private/ShaderParameterParser.cpp
+++ b/Engine/Source/Developer/ShaderCompilerCommon/Private/ShaderParameterParser.cpp
@@ -730,6 +730,48 @@ void FShaderParameterParser::RemoveMovingParametersFromSource(FString& Preproces
 	}
 }
 
+FString FShaderParameterParser::GenerateBindlessParameterDeclaration(const FParsedShaderParameter& ParsedParameter) const
+{
+	// NOTE: Macros AUTO_BINDLESS_SAMPLER_INDEX/VARIABLE and AUTO_BINDLESS_RESOURCE_INDEX/VARIABLE in BindlessResources.ush must be kept in sync with this function
+
+	const bool bIsSampler = (ParsedParameter.BindlessConversionType == EBindlessConversionType::Sampler);
+
+	const TCHAR* Kind = bIsSampler ? TEXT("Sampler") : TEXT("Resource");
+	const TCHAR* StorageClass = ParsedParameter.bGloballyCoherent ? TEXT("globallycoherent ") : TEXT("");
+
+	const FStringView Name = ParsedParameter.ParsedName;
+	const FStringView Type = ParsedParameter.ParsedType;
+
+	FString RewriteType(Type);
+
+	FString TypedefText = TEXT("");
+	if (ParsedParameter.BindlessConversionType == EBindlessConversionType::Resource)
+	{
+		RewriteType = FString::Printf(TEXT("SafeType%.*s"), Name.Len(), Name.GetData());
+		TypedefText = FString::Printf(TEXT("typedef %.*s %s;"), Type.Len(), Type.GetData(), *RewriteType);
+	}
+
+	TStringBuilder<512> Result;
+
+	// If we weren't going to be added to a root constant buffer, that means we need to declare our index before we declare our getter.
+	if (ParsedParameter.ConstantBufferParameterType == EShaderParameterType::Num)
+	{
+		// e.g. `uint BindlessResource_##Name;`
+		// or   `uint BindlessSampler_##Name;`
+		Result << TEXT("uint Bindless") << Kind << TEXT("_") << Name << TEXT(";");
+	}
+
+	Result << TypedefText;
+
+	// e.g. `Type GetBindlessResource##Name() { return GetResourceFromHeap(Type, BindlessResource_##Name); } static const Type Name = GetBindlessResource##Name()`
+	// or   `Type GetBindlessSampler##Name() { return GetSamplerFromHeap(Type, BindlessSampler_##Name); } static const Type Name = GetBindlessSampler##Name()`
+	Result << StorageClass << RewriteType << TEXT(" GetBindless") << Kind << Name << TEXT("()");
+	Result << TEXT("{ return Get") << Kind << TEXT("FromHeap(") << StorageClass << RewriteType << TEXT(", Bindless") << Kind << TEXT("_") << Name << TEXT("); } ");
+	Result << TEXT("static const ") << StorageClass << RewriteType << TEXT(" ") << Name << TEXT(" = GetBindless") << Kind << Name << TEXT("();");
+
+	return Result.ToString();
+}
+
 void FShaderParameterParser::ApplyBindlessModifications(FString& PreprocessedShaderSource)
 {
 	if (bBindlessResources || bBindlessSamplers)
@@ -755,38 +797,10 @@ void FShaderParameterParser::ApplyBindlessModifications(FString& PreprocessedSha
 
 			if (ParsedParameter.BindlessConversionType != EBindlessConversionType::None)
 			{
-				const bool bIsSampler = (ParsedParameter.BindlessConversionType == EBindlessConversionType::Sampler);
-
-				const TCHAR* IndexDefine = bIsSampler ? TEXT("AUTO_BINDLESS_SAMPLER_INDEX") : TEXT("AUTO_BINDLESS_RESOURCE_INDEX");
-				const TCHAR* VariableDefine = bIsSampler ? TEXT("AUTO_BINDLESS_SAMPLER_VARIABLE") : TEXT("AUTO_BINDLESS_RESOURCE_VARIABLE");
-				const TCHAR* StorageClass = ParsedParameter.bGloballyCoherent ? TEXT("globallycoherent ") : TEXT("");
-
-				const FStringView Name = ParsedParameter.ParsedName;
-				const FStringView Type = ParsedParameter.ParsedType;
-
-				FString RewriteType(Type);
-
-				FString TypedefText = TEXT("");
-				if (ParsedParameter.BindlessConversionType == EBindlessConversionType::Resource)
-				{
-					RewriteType = FString::Printf(TEXT("SafeType%.*s"), Name.Len(), Name.GetData());
-					TypedefText = FString::Printf(TEXT("typedef %.*s %s;"), Type.Len(), Type.GetData(), *RewriteType);
-				}
-
-				const FString IndexText = FString::Printf(TEXT("%s(%.*s);"), IndexDefine, Name.Len(), Name.GetData());
-				const FString VariableText = FString::Printf(TEXT("%s(%s%s, %.*s);"), VariableDefine, StorageClass, *RewriteType, Name.Len(), Name.GetData());
-
 				FShaderCodeModifications Modif;
 				Modif.CharOffsetStart = ParsedParameter.ParsedCharOffsetStart;
 				Modif.CharOffsetEnd = ParsedParameter.ParsedCharOffsetEnd + 1;
-
-				// If we weren't going to be added to a root constant buffer, that means we need to declare our index before we declare our getter.
-				if (ParsedParameter.ConstantBufferParameterType == EShaderParameterType::Num)
-				{
-					Modif.Replace += IndexText;
-				}
-				Modif.Replace += TypedefText;
-				Modif.Replace += VariableText;
+				Modif.Replace = GenerateBindlessParameterDeclaration(ParsedParameter);
 
 				Modifications.Add(Modif);
 			}
diff --git a/Engine/Source/Developer/ShaderCompilerCommon/Public/ShaderParameterParser.h b/Engine/Source/Developer/ShaderCompilerCommon/Public/ShaderParameterParser.h
--- a/Engine/Source/Developer/ShaderCompilerCommon/Public/ShaderParameterParser.h
+++ b/Engine/Source/Developer/ShaderCompilerCommon/Public/ShaderParameterParser.h
@@ -66,7 +66,7 @@ public:
 	FShaderParameterParser();
 	FShaderParameterParser(const TCHAR* InConstantBufferType);
 	FShaderParameterParser(const TCHAR* InConstantBufferType, TArrayView<const TCHAR* const> InExtraSRVTypes, TArrayView<const TCHAR* const> InExtraUAVTypes);
-	~FShaderParameterParser();
+	virtual ~FShaderParameterParser();
 
 	/** Parses the preprocessed shader code and applies the necessary modifications to it. */
 	bool ParseAndModify(
@@ -166,6 +166,12 @@ private:
 
 	void ExtractFileAndLine(int32 PragamLineoffset, int32 LineOffset, FString& OutFile, FString& OutLine) const;
 
+	/**
+	* Generates shader source code to declare a bindless resource or sampler (for automatic bindless conversion).
+	* May be overriden to allow custom implementations for different platforms.
+	*/
+	virtual FString GenerateBindlessParameterDeclaration(const FParsedShaderParameter& ParsedParameter) const;
+
 	const TCHAR* const ConstantBufferType = nullptr;
 
 	const TArray<const TCHAR*> ExtraSRVTypes;
