Include SPIR-V header from SPIRV-Headers folder instead of SPIRV-Reflect.

#rnx

diff --git a/Engine/Source/Developer/ShaderCompilerCommon/Public/SpirvCommon.h b/Engine/Source/Developer/ShaderCompilerCommon/Public/SpirvCommon.h
--- a/Engine/Source/Developer/ShaderCompilerCommon/Public/SpirvCommon.h
+++ b/Engine/Source/Developer/ShaderCompilerCommon/Public/SpirvCommon.h
@@ -8,7 +8,7 @@
 #include "Misc/AssertionMacros.h"
 
 THIRD_PARTY_INCLUDES_START
-	#include "spirv_reflect.h"
+	#include <spirv/unified1/spirv.h>
 THIRD_PARTY_INCLUDES_END
 
 
diff --git a/Engine/Source/Developer/ShaderCompilerCommon/Public/SpirvReflectCommon.h b/Engine/Source/Developer/ShaderCompilerCommon/Public/SpirvReflectCommon.h
--- a/Engine/Source/Developer/ShaderCompilerCommon/Public/SpirvReflectCommon.h
+++ b/Engine/Source/Developer/ShaderCompilerCommon/Public/SpirvReflectCommon.h
@@ -9,6 +9,7 @@
 
 #if PLATFORM_MAC || PLATFORM_WINDOWS || PLATFORM_LINUX
 THIRD_PARTY_INCLUDES_START
+	#include <spirv/unified1/spirv.h>
 	#include "spirv_reflect.h"
 THIRD_PARTY_INCLUDES_END
 #endif
diff --git a/Engine/Source/Developer/ShaderCompilerCommon/ShaderCompilerCommon.Build.cs b/Engine/Source/Developer/ShaderCompilerCommon/ShaderCompilerCommon.Build.cs
--- a/Engine/Source/Developer/ShaderCompilerCommon/ShaderCompilerCommon.Build.cs
+++ b/Engine/Source/Developer/ShaderCompilerCommon/ShaderCompilerCommon.Build.cs
@@ -31,5 +31,6 @@ public class ShaderCompilerCommon : ModuleRules
 		// We only need a header containing definitions
 		PublicSystemIncludePaths.Add(Path.Combine(Target.UEThirdPartySourceDirectory, "hlslcc/hlslcc/src/hlslcc_lib"));
 		PublicSystemIncludePaths.Add(Path.Combine(Target.UEThirdPartySourceDirectory, "SPIRV-Reflect/SPIRV-Reflect"));
+		PublicSystemIncludePaths.Add(Path.Combine(Target.UEThirdPartySourceDirectory, "ShaderConductor/ShaderConductor/External/SPIRV-Headers/include"));
 	}
 }
