[shaders] strip out some 5.2 deprecations

diff --git a/Engine/Source/Developer/ShaderCompilerCommon/Private/ShaderCompilerCommon.cpp b/Engine/Source/Developer/ShaderCompilerCommon/Private/ShaderCompilerCommon.cpp
--- a/Engine/Source/Developer/ShaderCompilerCommon/Private/ShaderCompilerCommon.cpp
+++ b/Engine/Source/Developer/ShaderCompilerCommon/Private/ShaderCompilerCommon.cpp
@@ -1142,7 +1142,7 @@ void TransformStringIntoCharacterArray(FString& PreprocessedShaderSource)
 
 ///////////////////////////////////////////////////////////////////////////////////////////////////////////////////////////////////////
 
-FString CreateShaderCompilerWorkerDirectCommandLine(const FShaderCompilerInput& Input, uint32 CCFlags)
+FString CreateShaderCompilerWorkerDirectCommandLine(const FShaderCompilerInput& Input, uint32 CCFlags = 0)
 {
 	FString Text(TEXT("-directcompile -format="));
 	Text += Input.ShaderFormat.GetPlainNameString();
@@ -1439,38 +1439,6 @@ void CompileOfflineMali(const FShaderCompilerInput& Input, FShaderCompilerOutput
 	}
 }
 
-
-FString GetDumpDebugUSFContents(const FShaderCompilerInput& Input, const FString& Source, uint32 HlslCCFlags)
-{
-	UE::ShaderCompilerCommon::FDebugShaderDataOptions DebugDataOptions;
-	DebugDataOptions.HlslCCFlags = HlslCCFlags;
-	return UE::ShaderCompilerCommon::GetDebugShaderContents(Input, Source, DebugDataOptions);
-}
-
-void DumpDebugUSF(const FShaderCompilerInput& Input, const ANSICHAR* Source, uint32 HlslCCFlags, const TCHAR* OverrideBaseFilename)
-{
-	PRAGMA_DISABLE_DEPRECATION_WARNINGS
-	FString NewSource = Source ? Source : "";
-	FString Contents = GetDumpDebugUSFContents(Input, NewSource, HlslCCFlags);
-	DumpDebugUSF(Input, NewSource, HlslCCFlags, OverrideBaseFilename);
-	PRAGMA_ENABLE_DEPRECATION_WARNINGS
-}
-
-void DumpDebugUSF(const FShaderCompilerInput& Input, const FString& Source, uint32 HlslCCFlags, const TCHAR* OverrideBaseFilename)
-{
-	PRAGMA_DISABLE_DEPRECATION_WARNINGS
-	FString BaseSourceFilename = (OverrideBaseFilename && *OverrideBaseFilename) ? OverrideBaseFilename : *Input.GetSourceFilename();
-	FString Filename = Input.DumpDebugInfoPath / BaseSourceFilename;
-
-	if (TUniquePtr<FArchive> FileWriter = TUniquePtr<FArchive>(IFileManager::Get().CreateFileWriter(*Filename)))
-	{
-		FString Contents = GetDumpDebugUSFContents(Input, Source, HlslCCFlags);
-		FileWriter->Serialize(TCHAR_TO_ANSI(*Contents), Contents.Len());
-		FileWriter->Close();
-	}
-	PRAGMA_ENABLE_DEPRECATION_WARNINGS
-}
-
 namespace UE::ShaderCompilerCommon
 {
 	void DumpDebugShaderData(const FShaderCompilerInput& Input, const FString& PreprocessedSource, const FDebugShaderDataOptions& Options)
@@ -1480,22 +1448,14 @@ namespace UE::ShaderCompilerCommon
 			FString Prefix = (Options.FilenamePrefix && *Options.FilenamePrefix) ? Options.FilenamePrefix : FString();
 			FString BaseSourceFilename = Prefix + ((Options.OverrideBaseFilename && *Options.OverrideBaseFilename) ? Options.OverrideBaseFilename : *Input.GetSourceFilename());
 
-			FString ModifiedPreprocessedSource = GetDebugShaderContents(Input, PreprocessedSource, Options);
-			PRAGMA_DISABLE_DEPRECATION_WARNINGS
-			// note: DumpDebugUSF should be internalized once the deprecation window ends 
-			// (i.e. still should remain in this file but not be declared in the header)
-			// until that time we need to disable deprecation warnings here.
-			DumpDebugUSF(Input, ModifiedPreprocessedSource, Options.HlslCCFlags, *BaseSourceFilename);
-			PRAGMA_ENABLE_DEPRECATION_WARNINGS
+			UE::ShaderCompilerCommon::FDebugShaderDataOptions DebugDataOptions;
+			DebugDataOptions.HlslCCFlags = Options.HlslCCFlags;
+			FString Contents = UE::ShaderCompilerCommon::GetDebugShaderContents(Input, PreprocessedSource, DebugDataOptions);
+			FFileHelper::SaveStringToFile(Contents, *(Input.DumpDebugInfoPath / BaseSourceFilename));
 
 			if (Input.bGenerateDirectCompileFile && !Options.bSkipDirectCompileTxt)
 			{
-				// note: CreateShaderCompileWorkerDirectCommandLine should be internalized once the deprecation window ends
-				// (i.e. still should remain in this file but not be declared in the header)
-				// until that time we need to disable deprecation warnings here.
-				PRAGMA_DISABLE_DEPRECATION_WARNINGS
 				FFileHelper::SaveStringToFile(CreateShaderCompilerWorkerDirectCommandLine(Input), *(Input.DumpDebugInfoPath / FString::Printf(TEXT("%sDirectCompile.txt"), *Prefix)));
-				PRAGMA_ENABLE_DEPRECATION_WARNINGS
 			}
 		}
 	}
@@ -1511,12 +1471,7 @@ namespace UE::ShaderCompilerCommon
 		Contents += TEXT("\n");
 		Contents += CrossCompiler::CreateResourceTableFromEnvironment(Input.Environment);
 		Contents += TEXT("#if 0 /*DIRECT COMPILE*/\n");
-		// note: CreateShaderCompileWorkerDirectCommandLine should be internalized once the deprecation window ends
-		// (i.e. still should remain in this file but not be declared in the header)
-		// until that time we need to disable deprecation warnings here.
-		PRAGMA_DISABLE_DEPRECATION_WARNINGS	
 		Contents += CreateShaderCompilerWorkerDirectCommandLine(Input, Options.HlslCCFlags);
-		PRAGMA_ENABLE_DEPRECATION_WARNINGS
 		Contents += TEXT("\n#endif /*DIRECT COMPILE*/\n");
 		if (!Input.DebugDescription.IsEmpty())
 		{
diff --git a/Engine/Source/Developer/ShaderCompilerCommon/Public/ShaderCompilerCommon.h b/Engine/Source/Developer/ShaderCompilerCommon/Public/ShaderCompilerCommon.h
--- a/Engine/Source/Developer/ShaderCompilerCommon/Public/ShaderCompilerCommon.h
+++ b/Engine/Source/Developer/ShaderCompilerCommon/Public/ShaderCompilerCommon.h
@@ -244,25 +244,12 @@ extern SHADERCOMPILERCOMMON_API bool RemoveUnusedInputs(FString& InOutSourceCode
 
 extern SHADERCOMPILERCOMMON_API bool ConvertFromFP32ToFP16(FString& InOutSourceCode, TArray<FString>& OutErrors);
 
-UE_DEPRECATED(5.2, "This function will become internal; please use the UE::ShaderCompilerCommon::DumpDebugShaderData function for dumping common debug info.")
-extern SHADERCOMPILERCOMMON_API FString CreateShaderCompilerWorkerDirectCommandLine(const FShaderCompilerInput& Input, uint32 CCFlags = 0);
-
 enum class EShaderConductorTarget
 {
 	Dxil,
 	Spirv,
 };
- extern SHADERCOMPILERCOMMON_API void WriteShaderConductorCommandLine(const FShaderCompilerInput& Input, const FString& SourceFilename, EShaderConductorTarget Target);
-
-// Gets the string that DumpDebugUSF writes out
-UE_DEPRECATED(5.2, "Please use UE::ShaderCompilerCommon::GetDebugShaderContents")
-extern SHADERCOMPILERCOMMON_API FString GetDumpDebugUSFContents(const FShaderCompilerInput& Input, const FString& Source, uint32 HlslCCFlags);
-
-// Utility functions shared amongst all backends to write out a dumped USF
-UE_DEPRECATED(5.2, "Use DumpDebugShaderData to handle all debug dumps from compiler backends.")
-extern SHADERCOMPILERCOMMON_API void DumpDebugUSF(const FShaderCompilerInput& Input, const ANSICHAR* Source, uint32 HlslCCFlags = 0, const TCHAR* OverrideBaseFilename = nullptr);
-UE_DEPRECATED(5.2, "Use DumpDebugShaderData to handle all debug dumps from compiler backends.")
-extern SHADERCOMPILERCOMMON_API void DumpDebugUSF(const FShaderCompilerInput& Input, const FString& Source, uint32 HlslCCFlags = 0, const TCHAR* OverrideBaseFilename = nullptr);
+extern SHADERCOMPILERCOMMON_API void WriteShaderConductorCommandLine(const FShaderCompilerInput& Input, const FString& SourceFilename, EShaderConductorTarget Target);
 
 extern SHADERCOMPILERCOMMON_API void DumpDebugShaderText(const FShaderCompilerInput& Input, const FString& InSource, const FString& FileExtension);
 extern SHADERCOMPILERCOMMON_API void DumpDebugShaderText(const FShaderCompilerInput& Input, ANSICHAR* InSource, int32 InSourceLength, const FString& FileExtension); 
diff --git a/Engine/Source/Runtime/Engine/Private/Materials/MaterialShader.cpp b/Engine/Source/Runtime/Engine/Private/Materials/MaterialShader.cpp
--- a/Engine/Source/Runtime/Engine/Private/Materials/MaterialShader.cpp
+++ b/Engine/Source/Runtime/Engine/Private/Materials/MaterialShader.cpp
@@ -2805,10 +2805,6 @@ const FMemoryImageString* FMaterialShaderMap::GetShaderSource(const FVertexFacto
 	return nullptr;
 }
 
-const FMemoryImageString* FMaterialShaderMap::GetShaderSource(const FName VertexFactoryName, const FName ShaderTypeName) const
-{
-	return GetShaderSource(FindVertexFactoryType(VertexFactoryName), FindShaderTypeByName(ShaderTypeName), /* PermutationId */ 0);
-}
 #endif // WITH_EDITOR
 
 void FMaterialShaderMap::GetShaderList(TMap<FShaderId, TShaderRef<FShader>>& OutShaders) const
diff --git a/Engine/Source/Runtime/Engine/Public/MaterialShared.h b/Engine/Source/Runtime/Engine/Public/MaterialShared.h
--- a/Engine/Source/Runtime/Engine/Public/MaterialShared.h
+++ b/Engine/Source/Runtime/Engine/Public/MaterialShared.h
@@ -1263,8 +1263,6 @@ public:
 
 #if WITH_EDITOR
 	ENGINE_API const FMemoryImageString* GetShaderSource(const FVertexFactoryType* VertexFactoryType, const FShaderType* ShaderType, int32 PermutationId) const;
-	UE_DEPRECATED(5.2, "Use the GetShaderSource overload accepting a FVertexFactoryType/FShaderType/permutation ID; this one doesn't correctly account for permutations.")
-	ENGINE_API const FMemoryImageString* GetShaderSource(FName VertexFactoryName, FName ShaderTypeName) const;
 #endif
 
 	/** Builds a list of the shaders in a shader map. */
