Unify shader debug name in shader code output for DXC backends to improve cache hits.

The shader code output contains a little bit of plaintext such as I/O identifiers and the shader name, which was generated for the DXC backend in a way that prevents caching.
This unifies the shader name in the output code and improves the cache hit ratio for a Lyra cook from 0% to roughly 50% cache hits (for certain platforms).

#rnx

diff --git a/Engine/Shaders/Public/Platform/GL/GLCommon.ush b/Engine/Shaders/Public/Platform/GL/GLCommon.ush
--- a/Engine/Shaders/Public/Platform/GL/GLCommon.ush
+++ b/Engine/Shaders/Public/Platform/GL/GLCommon.ush
@@ -5,7 +5,7 @@
 =============================================================================*/
 
 // Update this GUID to improve shader recompilation for GL only shaders
-// GUID = 17fde71b-0d49-409d-bb94-3a9f8c764098
+// GUID = C8B835E4-AF1B-4D8E-BD49-CBAAA96D6A0B
 
 #pragma once
 
diff --git a/Engine/Shaders/Public/Platform/Metal/MetalCommon.ush b/Engine/Shaders/Public/Platform/Metal/MetalCommon.ush
--- a/Engine/Shaders/Public/Platform/Metal/MetalCommon.ush
+++ b/Engine/Shaders/Public/Platform/Metal/MetalCommon.ush
@@ -5,7 +5,7 @@
 =============================================================================*/
 
 // Update this GUID to improve shader recompilation for Metal only shaders
-// GUID = A13B6FF9-78F9-41FE-8CC9-D532FBE70B28
+// GUID = F1C3BA6E-2148-457F-BA3C-80F387FF83C0
 
 #pragma once
 
diff --git a/Engine/Shaders/Public/Platform/Vulkan/VulkanCommon.ush b/Engine/Shaders/Public/Platform/Vulkan/VulkanCommon.ush
--- a/Engine/Shaders/Public/Platform/Vulkan/VulkanCommon.ush
+++ b/Engine/Shaders/Public/Platform/Vulkan/VulkanCommon.ush
@@ -5,7 +5,7 @@
 =============================================================================*/
 
 // Update this GUID to improve shader recompilation for Vulkan only shaders
-// GUID = 83946375-67fe-4f8f-be31-f9c03b9426ad
+// GUID = 16E30B80-05A3-47B5-AEDC-66036D3E028C
 
 #pragma once
 
diff --git a/Engine/Source/Developer/Apple/MetalShaderFormat/Private/MetalDerivedData.cpp b/Engine/Source/Developer/Apple/MetalShaderFormat/Private/MetalDerivedData.cpp
--- a/Engine/Source/Developer/Apple/MetalShaderFormat/Private/MetalDerivedData.cpp
+++ b/Engine/Source/Developer/Apple/MetalShaderFormat/Private/MetalDerivedData.cpp
@@ -1138,7 +1138,7 @@ bool DoCompileMetalShader(
 				CCHeaderWriter.WriteSideTable(TEXT("spvBufferSizeConstants"), SideTableIndex);
 			}
 
-			CCHeaderWriter.WriteSourceInfo(*Input.GetSourceFilename(), *Input.EntryPointName, *Input.DebugGroupName);
+			CCHeaderWriter.WriteSourceInfo(*Input.VirtualSourceFilePath, *Input.EntryPointName);
 			CCHeaderWriter.WriteCompilerInfo();
 
 			FString MetaData = CCHeaderWriter.ToString();
diff --git a/Engine/Source/Developer/ShaderCompilerCommon/Private/HlslccHeaderWriter.cpp b/Engine/Source/Developer/ShaderCompilerCommon/Private/HlslccHeaderWriter.cpp
--- a/Engine/Source/Developer/ShaderCompilerCommon/Private/HlslccHeaderWriter.cpp
+++ b/Engine/Source/Developer/ShaderCompilerCommon/Private/HlslccHeaderWriter.cpp
@@ -21,18 +21,11 @@ namespace CrossCompiler
 		MetaData += FString::Printf(Fmt, Args...);
 	}
 
-	void FHlslccHeaderWriter::WriteSourceInfo(const TCHAR* SourceName, const TCHAR* EntryPointName, const TCHAR* DebugGroupName)
+	void FHlslccHeaderWriter::WriteSourceInfo(const TCHAR* VirtualSourceFilePath, const TCHAR* EntryPointName)
 	{
-		check(SourceName != nullptr);
+		check(VirtualSourceFilePath != nullptr);
 		check(EntryPointName != nullptr);
-		if (DebugGroupName != nullptr)
-		{
-			Strings.SourceInfo = FString::Printf(TEXT("%s/%s:%s"), DebugGroupName, SourceName, EntryPointName);
-		}
-		else
-		{
-			Strings.SourceInfo = FString::Printf(TEXT("%s:%s"), SourceName, EntryPointName);
-		}
+		Strings.SourceInfo = FString::Printf(TEXT("%s:%s"), VirtualSourceFilePath, EntryPointName);
 	}
 
 	void FHlslccHeaderWriter::WriteCompilerInfo(const TCHAR* CompilerName)
diff --git a/Engine/Source/Developer/ShaderCompilerCommon/Public/HlslccHeaderWriter.h b/Engine/Source/Developer/ShaderCompilerCommon/Public/HlslccHeaderWriter.h
--- a/Engine/Source/Developer/ShaderCompilerCommon/Public/HlslccHeaderWriter.h
+++ b/Engine/Source/Developer/ShaderCompilerCommon/Public/HlslccHeaderWriter.h
@@ -22,7 +22,7 @@ namespace CrossCompiler
 	{
 	public:
 
-		void WriteSourceInfo(const TCHAR* SourceName, const TCHAR* EntryPointName, const TCHAR* DebugGroupName = nullptr);
+		void WriteSourceInfo(const TCHAR* VirtualSourceFilePath, const TCHAR* EntryPointName);
 		void WriteCompilerInfo(const TCHAR* CompilerName = TEXT("ShaderConductor"));
 
 		void WriteInputAttribute(const SpvReflectInterfaceVariable& Attribute);
diff --git a/Engine/Source/Developer/ShaderFormatOpenGL/Private/OpenGLShaderCompiler.cpp b/Engine/Source/Developer/ShaderFormatOpenGL/Private/OpenGLShaderCompiler.cpp
--- a/Engine/Source/Developer/ShaderFormatOpenGL/Private/OpenGLShaderCompiler.cpp
+++ b/Engine/Source/Developer/ShaderFormatOpenGL/Private/OpenGLShaderCompiler.cpp
@@ -3240,7 +3240,7 @@ static bool CompileToGlslWithShaderConductor(
 		}
 
 		// Generate meta data for CCHeader
-		CCHeaderWriter.WriteSourceInfo(*Input.GetSourceFilename(), *Input.EntryPointName, *Input.DebugGroupName);
+		CCHeaderWriter.WriteSourceInfo(*Input.VirtualSourceFilePath, *Input.EntryPointName);
 		CCHeaderWriter.WriteCompilerInfo();
 
 		if (GlslSource.length() > 0)
diff --git a/Engine/Source/Developer/VulkanShaderFormat/Private/VulkanShaderCompiler.cpp b/Engine/Source/Developer/VulkanShaderFormat/Private/VulkanShaderCompiler.cpp
--- a/Engine/Source/Developer/VulkanShaderFormat/Private/VulkanShaderCompiler.cpp
+++ b/Engine/Source/Developer/VulkanShaderFormat/Private/VulkanShaderCompiler.cpp
@@ -1885,7 +1885,7 @@ static bool BuildShaderOutputFromSpirv(
 	// Build final shader output with meta data
 	FString DebugName = Input.DumpDebugInfoPath.Right(Input.DumpDebugInfoPath.Len() - Input.DumpDebugInfoRootPath.Len());
 
-	CCHeaderWriter.WriteSourceInfo(*Input.GetSourceFilename(), *Input.EntryPointName, *Input.DebugGroupName);
+	CCHeaderWriter.WriteSourceInfo(*Input.VirtualSourceFilePath, *Input.EntryPointName);
 	CCHeaderWriter.WriteCompilerInfo();
 
 	const FString MetaData = CCHeaderWriter.ToString();
