[shaders] use ImageGetDigestStream for hashing DXC DLLs (to exclude signature/certificate data) instead of hashing the entire binary. This allows us to get the same version # for both signed and unsigned DLLs, and so fixes recompilation of shaders on initial startup of packaged release builds. Hashing was changed to xxhash which provides an iterative hashing API as required by this change (cityhash does not have such an API).

diff --git a/Engine/Source/Developer/ShaderCompilerCommon/Private/DXCWrapper.cpp b/Engine/Source/Developer/ShaderCompilerCommon/Private/DXCWrapper.cpp
--- a/Engine/Source/Developer/ShaderCompilerCommon/Private/DXCWrapper.cpp
+++ b/Engine/Source/Developer/ShaderCompilerCommon/Private/DXCWrapper.cpp
@@ -5,16 +5,26 @@
 #include "HAL/FileManager.h"
 #include "Templates/RefCounting.h"
 #include "Misc/Paths.h"
-#include "Hash/CityHash.h"
+#include "Hash/xxhash.h"
 
 #if PLATFORM_WINDOWS
 #include "Windows/WindowsHWrapper.h"
+#include "ImageHlp.h"
 #endif // PLATFORM_WINDOWS
 
 static TRefCountPtr<FDllHandle> GDxilHandle;
 static TRefCountPtr<FDllHandle> GDxcHandle;
 static TRefCountPtr<FDllHandle> GShaderConductorHandle;
 
+#if PLATFORM_WINDOWS
+BOOL __stdcall DigestHash(DIGEST_HANDLE Handle, PBYTE Data, DWORD Length)
+{
+	FXxHash64Builder& Hasher = *reinterpret_cast<FXxHash64Builder*>(Handle);
+	Hasher.Update(Data, Length);
+	return true;
+}
+#endif
+
 static uint64 GetLoadedModuleVersion(const TCHAR* ModuleName)
 {
 #if PLATFORM_WINDOWS
@@ -32,18 +42,25 @@ static uint64 GetLoadedModuleVersion(const TCHAR* ModuleName)
 	}
 	check(PathLen < UE_ARRAY_COUNT(DllPath));
 
-	uint64 FileVersion = 0;
-
-	TUniquePtr<FArchive> FileReader(IFileManager::Get().CreateFileReader(DllPath));
-	if (FileReader)
+	HANDLE DllHandle = ::CreateFileW(DllPath, GENERIC_READ, FILE_SHARE_READ | FILE_SHARE_WRITE, nullptr, OPEN_EXISTING, 0, nullptr);
+	FXxHash64Builder FileHashBuilder;
+	// ImageHlp APIs are not threadsafe; need a mutex around all calls.
+	// At the moment this is the only usage of said APIs in the engine; if other usage 
+	// is added this should be moved to a shared utility.
+	static FCriticalSection ImageHlpCs;
 	{
-		TArray<char> FileContents;
-		FileContents.SetNumUninitialized(FileReader->TotalSize());
-		FileReader->Serialize(FileContents.GetData(), FileContents.Num());
-		FileVersion = CityHash64(FileContents.GetData(), FileContents.Num());
+		FScopeLock Lock(&ImageHlpCs);
+		bool bResult = ImageGetDigestStream(
+			DllHandle,
+			CERT_PE_IMAGE_DIGEST_ALL_IMPORT_INFO | CERT_PE_IMAGE_DIGEST_RESOURCES,
+			DigestHash,
+			reinterpret_cast<DIGEST_HANDLE>(&FileHashBuilder));
+		check(bResult);
 	}
 
-	return FileVersion;
+	CloseHandle(DllHandle);
+
+	return FileHashBuilder.Finalize().Hash;
 #else // PLATFORM_WINDOWS
 	return 0;
 #endif // PLATFORM_WINDOWS
diff --git a/Engine/Source/Developer/ShaderCompilerCommon/ShaderCompilerCommon.Build.cs b/Engine/Source/Developer/ShaderCompilerCommon/ShaderCompilerCommon.Build.cs
--- a/Engine/Source/Developer/ShaderCompilerCommon/ShaderCompilerCommon.Build.cs
+++ b/Engine/Source/Developer/ShaderCompilerCommon/ShaderCompilerCommon.Build.cs
@@ -28,6 +28,11 @@ public class ShaderCompilerCommon : ModuleRules
 			AddEngineThirdPartyPrivateStaticDependencies(Target, "SPIRVReflect");
 		}
 
+		if (Target.Platform == UnrealTargetPlatform.Win64)
+		{
+			PublicSystemLibraries.Add("ImageHlp.lib");
+		}
+
 		// We only need a header containing definitions
 		PublicSystemIncludePaths.Add(Path.Combine(Target.UEThirdPartySourceDirectory, "hlslcc/hlslcc/src/hlslcc_lib"));
 		PublicSystemIncludePaths.Add(Path.Combine(Target.UEThirdPartySourceDirectory, "SPIRV-Reflect/SPIRV-Reflect"));
