Fix non-determinism in ShaderPrint hashes.

diff --git a/Engine/Source/Developer/ShaderCompilerCommon/Private/ShaderCompilerCommon.cpp b/Engine/Source/Developer/ShaderCompilerCommon/Private/ShaderCompilerCommon.cpp
--- a/Engine/Source/Developer/ShaderCompilerCommon/Private/ShaderCompilerCommon.cpp
+++ b/Engine/Source/Developer/ShaderCompilerCommon/Private/ShaderCompilerCommon.cpp
@@ -1077,7 +1077,11 @@ void TransformStringIntoCharacterArray(FString& PreprocessedShaderSource)
 				Entry.Offset		= GlobalCount;
 				Entry.SourceText	= Text;
 				ConvertTextToAsciiCharacter(Entry.SourceText, Entry.ConvertedText, Entry.EncodedText);
-				Entry.Hash			= CityHash32((const char*)&Entry.SourceText.GetCharArray(), sizeof(FString::ElementType) * Entry.SourceText.Len());
+				Entry.Hash			= CityHash32((const char*)Entry.SourceText.GetCharArray().GetData(), sizeof(FString::ElementType) * Entry.SourceText.Len());
+
+				// Sanity check
+				uint32 HCheck = CityHash32((const char*)Entry.SourceText.GetCharArray().GetData(), sizeof(FString::ElementType) * Entry.SourceText.Len());
+				check(HCheck == Entry.Hash);
 
 				GlobalCount += Entry.ConvertedText.Len();
 
