Fix FTokenizer::ProcessDirective warning string formatting that results in significant memory bloat

Due to a trivial bug/typo, instead of reporting the token that caused a parsing error, almost the entire shader source code was pasted into the warning string.

diff --git a/Engine/Source/Developer/ShaderCompilerCommon/Private/HlslLexer.cpp b/Engine/Source/Developer/ShaderCompilerCommon/Private/HlslLexer.cpp
--- a/Engine/Source/Developer/ShaderCompilerCommon/Private/HlslLexer.cpp
+++ b/Engine/Source/Developer/ShaderCompilerCommon/Private/HlslLexer.cpp
@@ -1310,13 +1310,13 @@ namespace CrossCompiler
 			else
 			{
 				FString Directive = TEXT("#if 0") + Tokenizer.ReadToEndOfLine();
-				CompilerMessages.SourceWarning(*FString::Printf(TEXT("Unhandled preprocessor directive (%s); HlslParser requires preprocessed input!"), Tokenizer.Current));
+				CompilerMessages.SourceWarning(*FString::Printf(TEXT("Unhandled preprocessor directive (%.500s); HlslParser requires preprocessed input!"), *Directive));
 			}
 		}
 		else
 		{
-			FString Directive = TEXT("#") + Tokenizer.ReadToEndOfLine();
-			CompilerMessages.SourceWarning(*FString::Printf(TEXT("Unhandled token (%s); HlslParser requires preprocessed input!"), Tokenizer.Current));
+			FString Directive = Tokenizer.ReadToEndOfLine();
+			CompilerMessages.SourceWarning(*FString::Printf(TEXT("Unhandled token (%.500s); HlslParser requires preprocessed input!"), *Directive));
 		}
 
 		Tokenizer.SkipToNextLine();
