[shaders] minor change to shader parameter parser parsed param struct to store member data inline rather than a pointer (precursor change to adding serialization to this type)

diff --git a/Engine/Source/Developer/ShaderCompilerCommon/Private/ShaderCompilerCommon.cpp b/Engine/Source/Developer/ShaderCompilerCommon/Private/ShaderCompilerCommon.cpp
--- a/Engine/Source/Developer/ShaderCompilerCommon/Private/ShaderCompilerCommon.cpp
+++ b/Engine/Source/Developer/ShaderCompilerCommon/Private/ShaderCompilerCommon.cpp
@@ -910,7 +910,7 @@ void AddUnboundShaderParameterError(
 	check(CompilerInput.RootParametersStructure);
 
 	const FShaderParameterParser::FParsedShaderParameter& Member = ShaderParameterParser.FindParameterInfos(ParameterBindingName);
-	check(!Member.IsBindable());
+	check(!Member.bIsBindable);
 
 	FShaderCompilerError Error(FString::Printf(
 		TEXT("Error: Shader parameter %s could not be bound to %s's shader parameter structure %s."),
diff --git a/Engine/Source/Developer/ShaderCompilerCommon/Private/ShaderParameterParser.cpp b/Engine/Source/Developer/ShaderCompilerCommon/Private/ShaderParameterParser.cpp
--- a/Engine/Source/Developer/ShaderCompilerCommon/Private/ShaderParameterParser.cpp
+++ b/Engine/Source/Developer/ShaderCompilerCommon/Private/ShaderParameterParser.cpp
@@ -170,10 +170,13 @@ bool FShaderParameterParser::ParseParameters(
 				uint16 ByteOffset)
 			{
 				FParsedShaderParameter ParsedParameter;
-				ParsedParameter.Member = &Member;
+				ParsedParameter.bIsBindable = true;
 				ParsedParameter.ConstantBufferOffset = ByteOffset;
-				check(ParsedParameter.IsBindable());
-
+				ParsedParameter.BaseType = Member.GetBaseType();
+				ParsedParameter.PrecisionModifier = Member.GetPrecision();
+				ParsedParameter.NumRows = Member.GetNumRows();
+				ParsedParameter.NumColumns = Member.GetNumColumns();
+				ParsedParameter.MemberSize = Member.IsVariableNativeType() ? Member.GetMemberSize() : 0u;
 				ParsedParameters.Add(ShaderBindingName, ParsedParameter);
 			});
 	}
@@ -319,9 +322,9 @@ bool FShaderParameterParser::ParseParameters(
 						ParsedParameter = ParsedParameters.FindChecked(ParsedParameterKey);
 
 						// Erase the parameter to move it into the root constant buffer.
-						if (bNeedToMoveToRootConstantBuffer && ParsedParameter.IsBindable())
+						if (bNeedToMoveToRootConstantBuffer && ParsedParameter.bIsBindable)
 						{
-							const EUniformBufferBaseType BaseType = ParsedParameter.Member->GetBaseType();
+							const EUniformBufferBaseType BaseType = ParsedParameter.BaseType;
 							bMoveToRootConstantBuffer =
 								BaseType == UBMT_INT32 ||
 								BaseType == UBMT_UINT32 ||
@@ -1016,10 +1019,16 @@ void FShaderParameterParser::ValidateShaderParameterType(
 	if (!bBindlessHack)
 	{
 		FString ExpectedShaderType;
-		ParsedParameter.Member->GenerateShaderParameterType(ExpectedShaderType, bPlatformSupportsPrecisionModifier);
+		FShaderParametersMetadata::FMember::GenerateShaderParameterType(
+			ExpectedShaderType, 
+			bPlatformSupportsPrecisionModifier,
+			ParsedParameter.BaseType,
+			ParsedParameter.PrecisionModifier,
+			ParsedParameter.NumRows,
+			ParsedParameter.NumColumns);
 
-		const bool bShouldBeInt = ParsedParameter.Member->GetBaseType() == UBMT_INT32;
-		const bool bShouldBeUint = ParsedParameter.Member->GetBaseType() == UBMT_UINT32;
+		const bool bShouldBeInt = ParsedParameter.BaseType == UBMT_INT32;
+		const bool bShouldBeUint = ParsedParameter.BaseType == UBMT_UINT32;
 
 		// Match parsed type with expected shader type
 		bool bIsTypeCorrect = ParsedParameter.ParsedType == ExpectedShaderType;
@@ -1035,7 +1044,7 @@ void FShaderParameterParser::ValidateShaderParameterType(
 			};
 
 			// Accept half-precision floats when single-precision was requested
-			if (ParsedParameter.ParsedType.StartsWith(TEXT("half")) && ParsedParameter.Member->GetBaseType() == UBMT_FLOAT32)
+			if (ParsedParameter.ParsedType.StartsWith(TEXT("half")) && ParsedParameter.BaseType == UBMT_FLOAT32)
 			{
 				bIsTypeCorrect = CheckTypeCorrect(4, 5);
 			}
@@ -1095,7 +1104,7 @@ void FShaderParameterParser::ValidateShaderParameterType(
 	}
 
 	// Validate parameter size, in case this is an array.
-	if (!bBindlessHack && ReflectionSize > int32(ParsedParameter.Member->GetMemberSize()))
+	if (!bBindlessHack && ReflectionSize > int32(ParsedParameter.MemberSize))
 	{
 		FString CppCodeName = CompilerInput.RootParametersStructure->GetFullMemberCodeName(ParsedParameter.ConstantBufferOffset);
 
@@ -1105,7 +1114,7 @@ void FShaderParameterParser::ValidateShaderParameterType(
 			*ShaderBindingName,
 			ReflectionSize,
 			*CppCodeName,
-			ParsedParameter.Member->GetMemberSize());
+			ParsedParameter.MemberSize);
 		GetParameterFileAndLine(ParsedParameter, Error.ErrorVirtualFilePath, Error.ErrorLineString);
 
 		CompilerOutput.Errors.Add(Error);
diff --git a/Engine/Source/Developer/ShaderCompilerCommon/Public/ShaderParameterParser.h b/Engine/Source/Developer/ShaderCompilerCommon/Public/ShaderParameterParser.h
--- a/Engine/Source/Developer/ShaderCompilerCommon/Public/ShaderParameterParser.h
+++ b/Engine/Source/Developer/ShaderCompilerCommon/Public/ShaderParameterParser.h
@@ -25,7 +25,11 @@ public:
 	{
 	public:
 		/** Original information about the member. */
-		const FShaderParametersMetadata::FMember* Member = nullptr;
+		EUniformBufferBaseType BaseType;
+		EShaderPrecisionModifier::Type PrecisionModifier;
+		uint32 NumRows;
+		uint32 NumColumns;
+		uint32 MemberSize;
 
 		/** Information found about the member when parsing the preprocessed code. */
 		FStringView ParsedName; /** View into FShaderParameterParser::OriginalParsedShader */
@@ -41,12 +45,6 @@ public:
 			return !ParsedType.IsEmpty();
 		}
 
-		/** Returns whether the shader parameter is bindable to the shader parameter structure. */
-		bool IsBindable() const
-		{
-			return Member != nullptr;
-		}
-
 		int32 ParsedPragmaLineoffset = 0;
 		int32 ParsedLineOffset = 0;
 
@@ -58,6 +56,7 @@ public:
 		EShaderParameterType ConstantBufferParameterType{};
 
 		bool bGloballyCoherent = false;
+		bool bIsBindable = false;
 
 		friend class FShaderParameterParser;
 	};
diff --git a/Engine/Source/Runtime/RenderCore/Private/ShaderParameterMetadata.cpp b/Engine/Source/Runtime/RenderCore/Private/ShaderParameterMetadata.cpp
--- a/Engine/Source/Runtime/RenderCore/Private/ShaderParameterMetadata.cpp
+++ b/Engine/Source/Runtime/RenderCore/Private/ShaderParameterMetadata.cpp
@@ -68,22 +68,28 @@ TMap<FHashedName, FShaderParametersMetadata*>& FShaderParametersMetadata::GetNam
 	return NameStructMap;
 }
 
-void FShaderParametersMetadata::FMember::GenerateShaderParameterType(FString& Result, bool bSupportsPrecisionModifier) const
+void FShaderParametersMetadata::FMember::GenerateShaderParameterType(
+	FString& Result,
+	bool bSupportsPrecisionModifier,
+	EUniformBufferBaseType BaseType,
+	EShaderPrecisionModifier::Type PrecisionModifier,
+	uint32 NumRows,
+	uint32 NumColumns)
 {
-	switch (GetBaseType())
+	switch (BaseType)
 	{
 	case UBMT_INT32:   Result = TEXT("int"); break;
 	case UBMT_UINT32:  Result = TEXT("uint"); break;
 	case UBMT_FLOAT32:
-		if (GetPrecision() == EShaderPrecisionModifier::Float || !bSupportsPrecisionModifier)
+		if (PrecisionModifier == EShaderPrecisionModifier::Float || !bSupportsPrecisionModifier)
 		{
 			Result = TEXT("float");
 		}
-		else if (GetPrecision() == EShaderPrecisionModifier::Half)
+		else if (PrecisionModifier == EShaderPrecisionModifier::Half)
 		{
 			Result = TEXT("half");
 		}
-		else if (GetPrecision() == EShaderPrecisionModifier::Fixed)
+		else if (PrecisionModifier == EShaderPrecisionModifier::Fixed)
 		{
 			Result = TEXT("fixed");
 		}
@@ -93,16 +99,21 @@ void FShaderParametersMetadata::FMember::GenerateShaderParameterType(FString& Re
 	};
 
 	// Generate the type dimensions for vectors and matrices.
-	if (GetNumRows() > 1)
+	if (NumRows > 1)
 	{
-		Result = FString::Printf(TEXT("%s%ux%u"), *Result, GetNumRows(), GetNumColumns());
+		Result = FString::Printf(TEXT("%s%ux%u"), *Result, NumRows, NumColumns);
 	}
-	else if (GetNumColumns() > 1)
+	else if (NumColumns > 1)
 	{
-		Result = FString::Printf(TEXT("%s%u"), *Result, GetNumColumns());
+		Result = FString::Printf(TEXT("%s%u"), *Result, NumColumns);
 	}
 }
 
+void FShaderParametersMetadata::FMember::GenerateShaderParameterType(FString& Result, bool bSupportsPrecisionModifier) const
+{
+	GenerateShaderParameterType(Result, bSupportsPrecisionModifier, GetBaseType(), GetPrecision(), GetNumRows(), GetNumColumns());
+}
+
 void FShaderParametersMetadata::FMember::GenerateShaderParameterType(FString& Result, EShaderPlatform ShaderPlatform) const
 {
 	GenerateShaderParameterType(Result, SupportShaderPrecisionModifier(ShaderPlatform));
@@ -530,10 +541,7 @@ void FShaderParametersMetadata::InitializeLayout(FRHIUniformBufferLayoutInitiali
 			BaseType == UBMT_SRV ||
 			BaseType == UBMT_SAMPLER);
 		const bool bIsRDGResource = IsRDGResourceReferenceShaderParameterType(BaseType);
-		const bool bIsVariableNativeType = (
-			BaseType == UBMT_INT32 ||
-			BaseType == UBMT_UINT32 ||
-			BaseType == UBMT_FLOAT32);
+		const bool bIsVariableNativeType = CurrentMember.IsVariableNativeType();
 
 		LayoutInitializer.bHasNonGraphOutputs |= BaseType == UBMT_UAV;
 
diff --git a/Engine/Source/Runtime/RenderCore/Private/ShaderParameterStruct.cpp b/Engine/Source/Runtime/RenderCore/Private/ShaderParameterStruct.cpp
--- a/Engine/Source/Runtime/RenderCore/Private/ShaderParameterStruct.cpp
+++ b/Engine/Source/Runtime/RenderCore/Private/ShaderParameterStruct.cpp
@@ -80,10 +80,7 @@ struct FShaderParameterStructBindingContext
 			const bool bIsRDGResource =
 				IsRDGResourceReferenceShaderParameterType(BaseType);
 
-			const bool bIsVariableNativeType = (
-				BaseType == UBMT_INT32 ||
-				BaseType == UBMT_UINT32 ||
-				BaseType == UBMT_FLOAT32);
+			const bool bIsVariableNativeType = Member.IsVariableNativeType();
 
 			checkf(BaseType != UBMT_BOOL, TEXT("Should have failed in FShaderParametersMetadata::InitializeLayout()"));
 
diff --git a/Engine/Source/Runtime/RenderCore/Public/ShaderParameterMetadata.h b/Engine/Source/Runtime/RenderCore/Public/ShaderParameterMetadata.h
--- a/Engine/Source/Runtime/RenderCore/Public/ShaderParameterMetadata.h
+++ b/Engine/Source/Runtime/RenderCore/Public/ShaderParameterMetadata.h
@@ -212,10 +212,18 @@ public:
 		/** Returns the metadata of the struct. */
 		const FShaderParametersMetadata* GetStructMetadata() const { return Struct; }
 
+		inline bool IsVariableNativeType() const
+		{
+			return 
+				BaseType == UBMT_INT32 ||
+				BaseType == UBMT_UINT32 ||
+				BaseType == UBMT_FLOAT32;
+		}
+
 		/** Returns the size of the member. */
 		inline uint32 GetMemberSize() const
 		{
-			check(BaseType == UBMT_FLOAT32 || BaseType == UBMT_INT32 || BaseType == UBMT_UINT32);
+			check(IsVariableNativeType());
 			uint32 ElementSize = sizeof(uint32) * NumRows * NumColumns;
 
 			/** If this an array, the alignment of the element are changed. */
@@ -226,6 +234,14 @@ public:
 			return ElementSize;
 		}
 
+		static void GenerateShaderParameterType(
+			FString& Result,
+			bool bSupportsPrecisionModifier,
+			EUniformBufferBaseType BaseType,
+			EShaderPrecisionModifier::Type PrecisionModifier,
+			uint32 NumRows,
+			uint32 NumColumns
+		);
 		void GenerateShaderParameterType(FString& Result, bool bSupportsPrecisionModifier) const;
 		void GenerateShaderParameterType(FString& Result, EShaderPlatform ShaderPlatform) const;
 
