[shaders] build fixes

diff --git a/Engine/Source/Developer/ShaderCompilerCommon/Private/ShaderCompilerCommon.cpp b/Engine/Source/Developer/ShaderCompilerCommon/Private/ShaderCompilerCommon.cpp
--- a/Engine/Source/Developer/ShaderCompilerCommon/Private/ShaderCompilerCommon.cpp
+++ b/Engine/Source/Developer/ShaderCompilerCommon/Private/ShaderCompilerCommon.cpp
@@ -1389,7 +1389,7 @@ void CompileOfflineMali(const FShaderCompilerInput& Input, FShaderCompilerOutput
 
 namespace UE::ShaderCompilerCommon
 {
-	void BaseShaderFormat::OutputDebugData(const FString& InputHash, const FShaderCompilerInput& Input, const FShaderPreprocessOutput& PreprocessOutput, const FShaderCompilerOutput& Output) const
+	void FBaseShaderFormat::OutputDebugData(const FString& InputHash, const FShaderCompilerInput& Input, const FShaderPreprocessOutput& PreprocessOutput, const FShaderCompilerOutput& Output) const
 	{
 		DumpDebugShaderData(Input, PreprocessOutput.GetOriginalSource());
 		FString OutputHashFileName = FPaths::Combine(Input.DumpDebugInfoPath, TEXT("OutputHash.txt"));
diff --git a/Engine/Source/Runtime/RenderCore/Private/ShaderPreprocessTypes.cpp b/Engine/Source/Runtime/RenderCore/Private/ShaderPreprocessTypes.cpp
--- a/Engine/Source/Runtime/RenderCore/Private/ShaderPreprocessTypes.cpp
+++ b/Engine/Source/Runtime/RenderCore/Private/ShaderPreprocessTypes.cpp
@@ -1,3 +1,5 @@
+// Copyright Epic Games, Inc. All Rights Reserved.
+
 #include "ShaderPreprocessTypes.h"
 
 #include "ShaderCompilerCore.h"
