Fixed mismatch string formatting argument in shader compiler (for hlslccflags).

diff --git a/Engine/Source/Developer/ShaderCompilerCommon/Private/ShaderCompilerCommon.cpp b/Engine/Source/Developer/ShaderCompilerCommon/Private/ShaderCompilerCommon.cpp
--- a/Engine/Source/Developer/ShaderCompilerCommon/Private/ShaderCompilerCommon.cpp
+++ b/Engine/Source/Developer/ShaderCompilerCommon/Private/ShaderCompilerCommon.cpp
@@ -1140,7 +1140,7 @@ FString CreateShaderCompilerWorkerDirectCommandLine(const FShaderCompilerInput&
 	if (CCFlags)
 	{
 		Text += TEXT(" -hlslccflags=");
-		Text += FString::Printf(TEXT("%llu"), CCFlags);
+		Text += FString::Printf(TEXT("%u"), CCFlags);
 	}
 	// When we're running in directcompile mode, we don't to spam the crash reporter
 	Text += TEXT(" -nocrashreports");
