[shaders] debug output improvements/fixes
- remove redundant options struct initialization in DumpDebugShaderData - this was resulting in losing some of the options set in the input parameter options struct
- add explicit "DumpExtendedDebugShaderData" which includes the additional debug output functionality only previously available if an IShaderFormat implementation inherited from FBaseShaderFormat
- add "AdditionalOutputs" functionality which can be used by backends to dump additional (text-based) files following the file naming rules as dictated by the options struct with less boilerplate code
- add a helper for constructing a debug filename from the debug output options struct (some backends require constructing such filenames for i.e. contents of batchfiles; the helper reduces boilerplate code in such cases)
- fix condition for old path of dumping OutputHash.txt to not redundantly dump if shader format supports independent preprocessing and preprocessed cache is enabled

diff --git a/Engine/Source/Developer/ShaderCompilerCommon/Private/ShaderCompilerCommon.cpp b/Engine/Source/Developer/ShaderCompilerCommon/Private/ShaderCompilerCommon.cpp
--- a/Engine/Source/Developer/ShaderCompilerCommon/Private/ShaderCompilerCommon.cpp
+++ b/Engine/Source/Developer/ShaderCompilerCommon/Private/ShaderCompilerCommon.cpp
@@ -1399,15 +1399,65 @@ void CompileOfflineMali(const FShaderCompilerInput& Input, FShaderCompilerOutput
 	}
 }
 
+// sensible default path size; TStringBuilder will allocate if it needs to
+const FString GetDebugFileName(
+	const FShaderCompilerInput& Input, 
+	const UE::ShaderCompilerCommon::FDebugShaderDataOptions& Options, 
+	const TCHAR* BaseFilename)
+{
+	TStringBuilder<512> PathBuilder;
+	const TCHAR* Prefix = (Options.FilenamePrefix && *Options.FilenamePrefix) ? Options.FilenamePrefix : TEXT("");
+	FStringView Filename = (BaseFilename && *BaseFilename) ? BaseFilename : Input.GetSourceFilenameView();
+	FPathViews::Append(PathBuilder, Input.DumpDebugInfoPath, Prefix);
+	PathBuilder << Filename;
+	return PathBuilder.ToString();
+}
+
 namespace UE::ShaderCompilerCommon
 {
-	void FBaseShaderFormat::OutputDebugData(const FShaderCompilerInput& Input, const FShaderPreprocessOutput& PreprocessOutput, const FShaderCompilerOutput& Output) const
+	FString FDebugShaderDataOptions::GetDebugShaderPath(const FShaderCompilerInput& Input) const
+	{
+		return GetDebugFileName(Input, *this, OverrideBaseFilename);
+	}
+
+	void FBaseShaderFormat::OutputDebugData(
+		const FShaderCompilerInput& Input,
+		const FShaderPreprocessOutput& PreprocessOutput,
+		const FShaderCompilerOutput& Output) const
+	{
+		DumpExtendedDebugShaderData(Input, PreprocessOutput, Output);
+	}
+
+	void DumpDebugShaderData(const FShaderCompilerInput& Input, const FString& PreprocessedSource, const FDebugShaderDataOptions& Options)
+	{
+		if (!Input.DumpDebugInfoEnabled())
+		{
+			return;
+		}
+
+		FString Contents = UE::ShaderCompilerCommon::GetDebugShaderContents(Input, PreprocessedSource, Options);
+		FFileHelper::SaveStringToFile(Contents, *Options.GetDebugShaderPath(Input));
+
+		if (EnumHasAnyFlags(Input.DebugInfoFlags, EShaderDebugInfoFlags::DirectCompileCommandLine) && !Options.bSourceOnly)
+		{
+			FFileHelper::SaveStringToFile(CreateShaderCompilerWorkerDirectCommandLine(Input), *GetDebugFileName(Input, Options, TEXT("DirectCompile.txt")));
+		}
+
+		for (const FDebugShaderDataOptions::FAdditionalOutput& AdditionalOutput : Options.AdditionalOutputs)
+		{
+			FFileHelper::SaveStringToFile(AdditionalOutput.Data, *GetDebugFileName(Input, Options, AdditionalOutput.BaseFileName));
+		}
+	}
+
+	void DumpExtendedDebugShaderData(
+		const FShaderCompilerInput& Input,
+		const FShaderPreprocessOutput& PreprocessOutput,
+		const FShaderCompilerOutput& Output,
+		const FDebugShaderDataOptions& Options)
 	{
 		const FString& SourceToDump = Output.ModifiedShaderSource.IsEmpty() ? PreprocessOutput.GetUnstrippedSource() : Output.ModifiedShaderSource;
-		DumpDebugShaderData(Input, SourceToDump);
-		
-		FString OutputHashFileName = FPaths::Combine(Input.DumpDebugInfoPath, TEXT("OutputHash.txt"));
-		FFileHelper::SaveStringToFile(Output.OutputHash.ToString(), *OutputHashFileName, FFileHelper::EEncodingOptions::ForceAnsi);
+		DumpDebugShaderData(Input, SourceToDump, Options);
+		FFileHelper::SaveStringToFile(Output.OutputHash.ToString(), *GetDebugFileName(Input, Options, TEXT("OutputHash.txt")), FFileHelper::EEncodingOptions::ForceAnsi);
 
 		if (EnumHasAnyFlags(Input.DebugInfoFlags, EShaderDebugInfoFlags::Diagnostics))
 		{
@@ -1418,33 +1468,18 @@ namespace UE::ShaderCompilerCommon
 			}
 			if (!Merged.IsEmpty())
 			{
-				FFileHelper::SaveStringToFile(Merged, *(Input.DumpDebugInfoPath / TEXT("Diagnostics.txt")));
+				FFileHelper::SaveStringToFile(Merged, *GetDebugFileName(Input, Options, TEXT("Diagnostics.txt")), FFileHelper::EEncodingOptions::ForceAnsi);
 			}
 		}
 
 		if (EnumHasAnyFlags(Input.DebugInfoFlags, EShaderDebugInfoFlags::InputHash))
 		{
-			FString InputHashFileName = FPaths::Combine(Input.DumpDebugInfoPath, TEXT("InputHash.txt"));
-			FFileHelper::SaveStringToFile(LexToString(Input.Hash), *InputHashFileName, FFileHelper::EEncodingOptions::ForceAnsi);
+			FFileHelper::SaveStringToFile(LexToString(Input.Hash), *GetDebugFileName(Input, Options, TEXT("InputHash.txt")), FFileHelper::EEncodingOptions::ForceAnsi);
 		}
-	}
 
-	void DumpDebugShaderData(const FShaderCompilerInput& Input, const FString& PreprocessedSource, const FDebugShaderDataOptions& Options)
-	{
-		if (Input.DumpDebugInfoEnabled())
+		for (const FDebugShaderDataOptions::FAdditionalOutput& AdditionalOutput : Options.AdditionalOutputs)
 		{
-			FString Prefix = (Options.FilenamePrefix && *Options.FilenamePrefix) ? Options.FilenamePrefix : FString();
-			FString BaseSourceFilename = Prefix + ((Options.OverrideBaseFilename && *Options.OverrideBaseFilename) ? Options.OverrideBaseFilename : *Input.GetSourceFilename());
-
-			UE::ShaderCompilerCommon::FDebugShaderDataOptions DebugDataOptions;
-			DebugDataOptions.HlslCCFlags = Options.HlslCCFlags;
-			FString Contents = UE::ShaderCompilerCommon::GetDebugShaderContents(Input, PreprocessedSource, DebugDataOptions);
-			FFileHelper::SaveStringToFile(Contents, *(Input.DumpDebugInfoPath / BaseSourceFilename));
-
-			if (EnumHasAnyFlags(Input.DebugInfoFlags, EShaderDebugInfoFlags::DirectCompileCommandLine) && !Options.bSourceOnly)
-			{
-				FFileHelper::SaveStringToFile(CreateShaderCompilerWorkerDirectCommandLine(Input), *(Input.DumpDebugInfoPath / FString::Printf(TEXT("%sDirectCompile.txt"), *Prefix)));
-			}
+			FFileHelper::SaveStringToFile(AdditionalOutput.Data, *GetDebugFileName(Input, Options, AdditionalOutput.BaseFileName), FFileHelper::EEncodingOptions::ForceAnsi);
 		}
 	}
 
diff --git a/Engine/Source/Developer/ShaderCompilerCommon/Public/ShaderCompilerCommon.h b/Engine/Source/Developer/ShaderCompilerCommon/Public/ShaderCompilerCommon.h
--- a/Engine/Source/Developer/ShaderCompilerCommon/Public/ShaderCompilerCommon.h
+++ b/Engine/Source/Developer/ShaderCompilerCommon/Public/ShaderCompilerCommon.h
@@ -72,12 +72,6 @@ extern SHADERCOMPILERCOMMON_API int16 GetNumUniformBuffersUsed(const FShaderComp
 
 namespace UE::ShaderCompilerCommon
 {
-	class FBaseShaderFormat : public IShaderFormat
-	{
-	public:
-		virtual SHADERCOMPILERCOMMON_API void OutputDebugData(const FShaderCompilerInput& Input, const FShaderPreprocessOutput& PreprocessOutput, const FShaderCompilerOutput& Output) const override;
-	};
-
 	UE_DEPRECATED(5.3, "Deprecated; use FShaderCompilerInput::ShouldUseStableContantBuffer directly")
 	inline bool ShouldUseStableConstantBuffer(const FShaderCompilerInput& Input)
 	{
@@ -124,18 +118,27 @@ namespace UE::ShaderCompilerCommon
 
 	struct FDebugShaderDataOptions
 	{
+		struct FAdditionalOutput
+		{
+			const TCHAR* BaseFileName;
+			const TCHAR* Data;
+		};
 		uint32 HlslCCFlags = 0;
 		const TCHAR* OverrideBaseFilename = nullptr;
 		const TCHAR* FilenamePrefix = nullptr;
 		TFunction<FString()> AppendPreSource{};
 		TFunction<FString()> AppendPostSource{};
+		TArray<FAdditionalOutput> AdditionalOutputs;
 		union
 		{
 			bool bSourceOnly = false; // if true, will only output source .usf as directed and skip all other debug data artifacts
 			UE_DEPRECATED(5.3, "bSkipDirectCompileTxt is deprecated, use bSourceOnly flag instead")
 			bool bSkipDirectCompileTxt;
 		};
+
+		SHADERCOMPILERCOMMON_API FString GetDebugShaderPath(const FShaderCompilerInput& Input) const;
 	};
+
 	/*
 	 * Dumps common debug information (preprocessed .usf as constructed by GetDebugShaderContents, and a directcompile.txt file 
 	 * containing the commandline for launching ShaderCompileWorker manually) for the given shader compile input
@@ -146,6 +149,22 @@ namespace UE::ShaderCompilerCommon
 	 */
 	extern SHADERCOMPILERCOMMON_API void DumpDebugShaderData(const FShaderCompilerInput& Input, const FString& PreprocessedSource, const FDebugShaderDataOptions& Options = FDebugShaderDataOptions());
 
+	/*
+	 * Dumps extended debug information; including all outputs from DumpDebugShaderData as well as the following:
+	 *		- OutputHash.txt file containing the SHA hash of the shader job output
+	 *		- Diagnostics.txt file containing all errors/warnings encountered for the job
+	 *		  (if EShaderDebugInfoFlags::Diagnostics is set on Input.DebugInfoFlags)
+	 *		- InputHash.txt file containing the hash used as the key to the shader job cache
+	 *		  (if EShaderDebugInfoFlags::InputHash is set on Input.DebugInfoFlags and specified InputHash is non-empty)
+	 *		- any outputs specified on the AdditionalOutputs array in the Options struct
+	 * This is intended to be used by shader formats implementing the independent preprocessing API.
+	 */
+	extern SHADERCOMPILERCOMMON_API void DumpExtendedDebugShaderData(
+		const FShaderCompilerInput& Input,
+		const FShaderPreprocessOutput& PreprocessOutput,
+		const FShaderCompilerOutput& Output,
+		const UE::ShaderCompilerCommon::FDebugShaderDataOptions& Options = FDebugShaderDataOptions());
+
 	/*
 	 * Constructs the modified preprocessed source that would be dumped to a .usf file via DumpDebugShaderData, including the following additions:
 	 * - resource table debug information
@@ -157,6 +176,15 @@ namespace UE::ShaderCompilerCommon
 	 * @param	Options Options which can change behaviour of the debug dump; see above.
 	 */
 	extern SHADERCOMPILERCOMMON_API FString GetDebugShaderContents(const FShaderCompilerInput& Input, const FString& PreprocessedSource, const FDebugShaderDataOptions& Options = FDebugShaderDataOptions());
+
+	class FBaseShaderFormat : public IShaderFormat
+	{
+	public:
+		virtual SHADERCOMPILERCOMMON_API void OutputDebugData(
+			const FShaderCompilerInput& Input,
+			const FShaderPreprocessOutput& PreprocessOutput,
+			const FShaderCompilerOutput& Output) const override;
+	};
 }
 
 extern SHADERCOMPILERCOMMON_API void HandleReflectedGlobalConstantBufferMember(
diff --git a/Engine/Source/Runtime/RenderCore/Private/ShaderCore.cpp b/Engine/Source/Runtime/RenderCore/Private/ShaderCore.cpp
--- a/Engine/Source/Runtime/RenderCore/Private/ShaderCore.cpp
+++ b/Engine/Source/Runtime/RenderCore/Private/ShaderCore.cpp
@@ -1142,13 +1142,14 @@ void CompileShader(const TArray<const IShaderFormat*>& ShaderFormats, FShaderCom
 	Job.bSucceeded = Job.Output.bSucceeded;
 	if (Job.Input.DumpDebugInfoEnabled())
 	{
-		// if the preprocessed cache is disabled, dump debug output here, since we don't serialize preprocess output 
-		// back to the cooker in this case (if enabled this will occur in the job OnComplete callback)
-		// pass an empty string for input hash; this is not available here but also not important if the preprocessed
-		// cache is disabled
-		if (Compiler->SupportsIndependentPreprocessing() && !Job.Input.bCachePreprocessed)
+		if (Compiler->SupportsIndependentPreprocessing())
 		{
-			Compiler->OutputDebugData(Job.Input, Job.PreprocessOutput, Job.Output);
+			// if the shader format supports independent preprocessing and preprocessed cache is disabled, dump debug output here, since we
+			// don't serialize preprocess output back to the cooker from SCW (if enabled this will occur in the job OnComplete callback)
+			if (!Job.Input.bCachePreprocessed)
+			{
+				Compiler->OutputDebugData(Job.Input, Job.PreprocessOutput, Job.Output);
+			}
 		}
 		else
 		{
diff --git a/Engine/Source/Runtime/RenderCore/Public/ShaderCompilerCore.h b/Engine/Source/Runtime/RenderCore/Public/ShaderCompilerCore.h
--- a/Engine/Source/Runtime/RenderCore/Public/ShaderCompilerCore.h
+++ b/Engine/Source/Runtime/RenderCore/Public/ShaderCompilerCore.h
@@ -14,6 +14,7 @@
 #include "Misc/EnumClassFlags.h"
 #include "Misc/SecureHash.h"
 #include "Misc/Paths.h"
+#include "Misc/PathViews.h"
 #include "Misc/CoreStats.h"
 #include "ShaderCore.h"
 #include "ShaderParameterMetadata.h"
@@ -314,6 +315,11 @@ struct FShaderCompilerInput
 		return Name;
 	}
 
+	FStringView GetSourceFilenameView() const
+	{
+		return FPathViews::GetCleanFilename(VirtualSourceFilePath);
+	}
+
 	FString GetSourceFilename() const
 	{
 		return FPaths::GetCleanFilename(VirtualSourceFilePath);
