Fix crash on empty string of preprocessed shader source.

diff --git a/Engine/Source/Developer/ShaderCompilerCommon/Private/ShaderCompilerCommon.cpp b/Engine/Source/Developer/ShaderCompilerCommon/Private/ShaderCompilerCommon.cpp
--- a/Engine/Source/Developer/ShaderCompilerCommon/Private/ShaderCompilerCommon.cpp
+++ b/Engine/Source/Developer/ShaderCompilerCommon/Private/ShaderCompilerCommon.cpp
@@ -978,6 +978,12 @@ void ConvertTextToAsciiCharacter(const FString& InText, FString& OutText, FStrin
 // Simple token matching and expansion to replace TEXT macro into supported character string
 void TransformStringIntoCharacterArray(FString& PreprocessedShaderSource)
 {
+	// Early out if input is empty; '&PreprocessedShaderSource[0]' below does not return a valid pointer for empty FString
+	if (PreprocessedShaderSource.IsEmpty())
+	{
+		return;
+	}
+
 	struct FTextEntry
 	{
 		uint32  Index;
