Fix stripping of unused interpolator. HLSLParser used by the RemoveUnusedOutputs function wasn't able to parse bit_cast, nointrerp, ulong, and Constant tokens.

-Added support to write PSSL qualifier (similar to PrimitiveId)

diff --git a/Engine/Source/Developer/ShaderCompilerCommon/Private/HlslAST.cpp b/Engine/Source/Developer/ShaderCompilerCommon/Private/HlslAST.cpp
--- a/Engine/Source/Developer/ShaderCompilerCommon/Private/HlslAST.cpp
+++ b/Engine/Source/Developer/ShaderCompilerCommon/Private/HlslAST.cpp
@@ -847,13 +847,15 @@ namespace CrossCompiler
 			{
 				Writer << TEXT("centroid ");
 			}
-			if (bNoInterpolation)
+			if (NoInterpolatorType)
 			{
-				Writer << TEXT("nointerpolation ");
+				Writer << NoInterpolatorType;
+				Writer << (TCHAR)' ';
 			}
-			if (bNoPerspective)
+			if (NoPerspectiveType)
 			{
-				Writer << TEXT("noperspective ");
+				Writer << NoPerspectiveType;
+				Writer << (TCHAR)' ';
 			}
 			if (bSample)
 			{
diff --git a/Engine/Source/Developer/ShaderCompilerCommon/Private/HlslAST.h b/Engine/Source/Developer/ShaderCompilerCommon/Private/HlslAST.h
--- a/Engine/Source/Developer/ShaderCompilerCommon/Private/HlslAST.h
+++ b/Engine/Source/Developer/ShaderCompilerCommon/Private/HlslAST.h
@@ -531,14 +531,14 @@ namespace CrossCompiler
 					// Interpolation modifiers
 					uint32 bLinear : 1;
 					uint32 bCentroid : 1;
-					uint32 bNoInterpolation : 1;
-					uint32 bNoPerspective : 1;
 					uint32 bSample : 1;
 				};
 				uint32 Raw;
 			};
 
 			const TCHAR* PrimitiveType = nullptr;
+			const TCHAR* NoInterpolatorType = nullptr;
+			const TCHAR* NoPerspectiveType = nullptr;
 
 			FTypeQualifier();
 
diff --git a/Engine/Source/Developer/ShaderCompilerCommon/Private/HlslLexer.cpp b/Engine/Source/Developer/ShaderCompilerCommon/Private/HlslLexer.cpp
--- a/Engine/Source/Developer/ShaderCompilerCommon/Private/HlslLexer.cpp
+++ b/Engine/Source/Developer/ShaderCompilerCommon/Private/HlslLexer.cpp
@@ -311,6 +311,11 @@ namespace CrossCompiler
 				InsertToken(TEXT("uint64_t3x4"), EHlslToken::Uint64_t3x4);
 				InsertToken(TEXT("uint64_t4x4"), EHlslToken::Uint64_t4x4);
 
+				InsertToken(TEXT("ulong"), EHlslToken::Uint64_t);	
+				InsertToken(TEXT("ulong2"), EHlslToken::Uint64_t2);	
+				InsertToken(TEXT("ulong3"), EHlslToken::Uint64_t3);	
+				InsertToken(TEXT("ulong4"), EHlslToken::Uint64_t4);	
+
 				InsertToken(TEXT("half"), EHlslToken::Half);
 				InsertToken(TEXT("half1"), EHlslToken::Half1);
 				InsertToken(TEXT("half2"), EHlslToken::Half2);
diff --git a/Engine/Source/Developer/ShaderCompilerCommon/Private/HlslParser.cpp b/Engine/Source/Developer/ShaderCompilerCommon/Private/HlslParser.cpp
--- a/Engine/Source/Developer/ShaderCompilerCommon/Private/HlslParser.cpp
+++ b/Engine/Source/Developer/ShaderCompilerCommon/Private/HlslParser.cpp
@@ -560,22 +560,22 @@ namespace CrossCompiler
 						return ParseResultError();
 					}
 				}
-				else if (Token->String == TEXT("nointerpolation"))
+				else if (Token->String == TEXT("nointerpolation") || Token->String == TEXT("nointerp"))
 				{
 					Scanner.Advance();
 					++InterpolationNoInterpolationFound;
-					Qualifier->bNoInterpolation = true;
+					Qualifier->NoInterpolatorType = Allocator->Strdup(Token->String);
 					if (InterpolationNoInterpolationFound > 1)
 					{
 						Scanner.SourceError(TEXT("'nointerpolation' found more than once!\n"));
 						return ParseResultError();
 					}
 				}
-				else if (Token->String == TEXT("noperspective") || Token->String == TEXT("nopersp"))	// PSSL nopersp
+				else if (Token->String == TEXT("noperspective") || Token->String == TEXT("nopersp"))	
 				{
 					Scanner.Advance();
 					++InterpolationNoPerspectiveFound;
-					Qualifier->bNoPerspective = true;
+					Qualifier->NoPerspectiveType = Allocator->Strdup(Token->String);
 					if (InterpolationNoPerspectiveFound > 1)
 					{
 						Scanner.SourceError(TEXT("'noperspective' found more than once!\n"));
@@ -1926,7 +1926,7 @@ Done:
 			return EParseResult::Matched;
 		}
 
-		if (Parser.Scanner.MatchToken(EHlslToken::CBuffer))
+		if (Parser.Scanner.MatchToken(EHlslToken::ConstantBuffer) || Parser.Scanner.MatchToken(EHlslToken::CBuffer))
 		{
 			auto Result = ParseCBuffer(Parser, Allocator, OutNode);
 			if (Result == EParseResult::Error || Result == EParseResult::Matched)
