[shaders] debug output changes - add optional shader code binary output, fix redundant write of additional outputs

diff --git a/Engine/Source/Developer/ShaderCompilerCommon/Private/ShaderCompilerCommon.cpp b/Engine/Source/Developer/ShaderCompilerCommon/Private/ShaderCompilerCommon.cpp
--- a/Engine/Source/Developer/ShaderCompilerCommon/Private/ShaderCompilerCommon.cpp
+++ b/Engine/Source/Developer/ShaderCompilerCommon/Private/ShaderCompilerCommon.cpp
@@ -1448,11 +1448,6 @@ namespace UE::ShaderCompilerCommon
 		{
 			FFileHelper::SaveStringToFile(CreateShaderCompilerWorkerDirectCommandLine(Input), *GetDebugFileName(Input, Options, TEXT("DirectCompile.txt")));
 		}
-
-		for (const FDebugShaderDataOptions::FAdditionalOutput& AdditionalOutput : Options.AdditionalOutputs)
-		{
-			FFileHelper::SaveStringToFile(AdditionalOutput.Data, *GetDebugFileName(Input, Options, AdditionalOutput.BaseFileName));
-		}
 	}
 
 	void DumpExtendedDebugShaderData(
@@ -1483,6 +1478,23 @@ namespace UE::ShaderCompilerCommon
 			FFileHelper::SaveStringToFile(LexToString(Input.Hash), *GetDebugFileName(Input, Options, TEXT("InputHash.txt")), FFileHelper::EEncodingOptions::ForceAnsi);
 		}
 
+		if (EnumHasAnyFlags(Input.DebugInfoFlags, EShaderDebugInfoFlags::ShaderCodeBinary))
+		{
+			FString ShaderCodeFileName = *GetDebugFileName(Input, Options, TEXT("ShaderCode.bin"));
+			if (Output.ShaderCode.IsCompressed())
+			{
+				// always output decompressed code as it's slightly more useful for A/B comparisons
+				TArray<uint8> DecompressedCode;
+				DecompressedCode.SetNum(Output.ShaderCode.GetUncompressedSize());
+				bool bSucceed = FCompression::UncompressMemory(NAME_Oodle, DecompressedCode.GetData(), DecompressedCode.Num(), Output.ShaderCode.GetReadAccess().GetData(), Output.ShaderCode.GetShaderCodeSize());
+				FFileHelper::SaveArrayToFile(DecompressedCode, *ShaderCodeFileName);
+			}
+			else
+			{
+				FFileHelper::SaveArrayToFile(Output.ShaderCode.GetReadAccess(), *ShaderCodeFileName);
+			}
+		}
+
 		for (const FDebugShaderDataOptions::FAdditionalOutput& AdditionalOutput : Options.AdditionalOutputs)
 		{
 			FFileHelper::SaveStringToFile(AdditionalOutput.Data, *GetDebugFileName(Input, Options, AdditionalOutput.BaseFileName), FFileHelper::EEncodingOptions::ForceAnsi);
diff --git a/Engine/Source/Runtime/Engine/Private/ShaderCompiler/ShaderCompiler.cpp b/Engine/Source/Runtime/Engine/Private/ShaderCompiler/ShaderCompiler.cpp
--- a/Engine/Source/Runtime/Engine/Private/ShaderCompiler/ShaderCompiler.cpp
+++ b/Engine/Source/Runtime/Engine/Private/ShaderCompiler/ShaderCompiler.cpp
@@ -148,6 +148,12 @@ static TAutoConsoleVariable<bool> CVarDebugDumpJobDiagnostics(
 	TEXT("If true, all diagnostic messages (errors and warnings) for each shader job will be dumped alongside other debug data (in Diagnostics.txt)"),
 	ECVF_ReadOnly);
 
+static TAutoConsoleVariable<bool> CVarDebugDumpShaderCode(
+	TEXT("r.ShaderCompiler.DebugDumpShaderCode"),
+	false,
+	TEXT("If true, each shader job will dump a ShaderCode.bin containing the contents of the output shader code object (the contents of this can differ for each shader format; note that this is the data that is hashed to produce the OutputHash.txt file)"),
+	ECVF_ReadOnly);
+
 static TAutoConsoleVariable<bool> CVarCompileParallelInProcess(
 	TEXT("r.ShaderCompiler.ParallelInProcess"),
 	false,
@@ -4387,6 +4393,11 @@ EShaderDebugInfoFlags FShaderCompilingManager::GetDumpShaderDebugInfoFlags() con
 		Flags |= EShaderDebugInfoFlags::Diagnostics;
 	}
 
+	if (CVarDebugDumpShaderCode.GetValueOnAnyThread())
+	{
+		Flags |= EShaderDebugInfoFlags::ShaderCodeBinary;
+	}
+
 	return Flags;
 }
 
diff --git a/Engine/Source/Runtime/RenderCore/Public/ShaderCompilerCore.h b/Engine/Source/Runtime/RenderCore/Public/ShaderCompilerCore.h
--- a/Engine/Source/Runtime/RenderCore/Public/ShaderCompilerCore.h
+++ b/Engine/Source/Runtime/RenderCore/Public/ShaderCompilerCore.h
@@ -199,6 +199,7 @@ enum class EShaderDebugInfoFlags : uint8
 	DirectCompileCommandLine = 1 << 0,
 	InputHash = 1 << 1,
 	Diagnostics = 1 << 2,
+	ShaderCodeBinary = 1 << 3,
 };
 ENUM_CLASS_FLAGS(EShaderDebugInfoFlags)
 
