Support for slim ShaderSymbols.info / r.Shaders.SymbolsInfo=1
Platforms can return a platform hash and a blob of text data to go with that hash
These are stored one / line in ShaderSymbols.info in lieu of full PDB info
Useful when full shader PDBs are impractically large, info-only is much smaller

#tests test on cooking Lyra, various platforms, noddc and ddc

diff --git a/Engine/Source/Developer/ShaderCompilerCommon/Private/ShaderSymbolExport.cpp b/Engine/Source/Developer/ShaderCompilerCommon/Private/ShaderSymbolExport.cpp
--- a/Engine/Source/Developer/ShaderCompilerCommon/Private/ShaderSymbolExport.cpp
+++ b/Engine/Source/Developer/ShaderCompilerCommon/Private/ShaderSymbolExport.cpp
@@ -18,6 +18,7 @@ DECLARE_LOG_CATEGORY_CLASS(LogShaderSymbolExport, Display, Display);
 
 static const TCHAR* ZipFileBaseLeafName = TEXT("ShaderSymbols");
 static const TCHAR* ZipFileExtension = TEXT(".zip");
+static const TCHAR* InfoFileExtension = TEXT(".info");
 
 FShaderSymbolExport::FShaderSymbolExport(FName InShaderFormat)
 	: ShaderFormat(InShaderFormat)
@@ -26,26 +27,42 @@ FShaderSymbolExport::FShaderSymbolExport(FName InShaderFormat)
 
 FShaderSymbolExport::~FShaderSymbolExport() = default;
 
-static void DeleteExistingShaderZips(IPlatformFile& PlatformFile, const FString& Directory)
+static void DeleteExisting(IPlatformFile& PlatformFile, const FString& Directory, const TCHAR* BaseLeafName, const TCHAR* Extension)
 {
 	TArray<FString> ExistingZips;
-	PlatformFile.FindFiles(ExistingZips, *Directory, ZipFileExtension);
+	PlatformFile.FindFiles(ExistingZips, *Directory, Extension);
 
 	for (const FString& ZipFile : ExistingZips)
 	{
-		if (FPathViews::GetPathLeaf(ZipFile).StartsWith(ZipFileBaseLeafName))
+		if (FPathViews::GetPathLeaf(ZipFile).StartsWith(BaseLeafName))
 		{
 			PlatformFile.DeleteFile(*ZipFile);
 		}
 	}
 }
 
+static FString CreateNameAndDeleteOld(uint32 MultiprocessId, IPlatformFile& PlatformFile, const FString& ExportPath, const TCHAR* BaseLeafName, const TCHAR* Extension)
+{
+	FString Name;
+	if (MultiprocessId == 0)
+	{
+		DeleteExisting(PlatformFile, ExportPath, BaseLeafName, Extension);
+		Name = FString::Printf(TEXT("%s%s"), BaseLeafName, Extension);
+	}
+	else
+	{
+		Name = FString::Printf(TEXT("%s_%d%s"), BaseLeafName, MultiprocessId, Extension);
+	}
+	return Name;
+}
+
 void FShaderSymbolExport::Initialize()
 {
 	const bool bSymbolsEnabled = ShouldWriteShaderSymbols(ShaderFormat);
 	const bool bForceSymbols = FParse::Value(FCommandLine::Get(), TEXT("-ShaderSymbolsExport="), ExportPath);
+	const bool bSymbolsInfoEnabled = ShouldGenerateShaderSymbolsInfo(ShaderFormat);
 
-	if (bSymbolsEnabled || bForceSymbols)
+	if (bSymbolsEnabled || bForceSymbols || bSymbolsInfoEnabled)
 	{
 		// if no command line path is provided, look to the cvar first
 		if (ExportPath.IsEmpty())
@@ -72,23 +89,16 @@ void FShaderSymbolExport::Initialize()
 		}
 		else
 		{
+			// setup multiproc data in case we need it
+			uint32 MultiprocessId = 0;
+			FParse::Value(FCommandLine::Get(), TEXT("-MultiprocessId="), MultiprocessId);
+			bMultiprocessOwner = MultiprocessId == 0;
+
 			// Check if the export mode is to an uncompressed archive or loose files.
-			bool bExportAsZip = ShouldWriteShaderSymbolsAsZip(ShaderFormat);
-			if (bExportAsZip || FParse::Param(FCommandLine::Get(), TEXT("ShaderSymbolsExportZip")))
+			const bool bExportAsZip = ShouldWriteShaderSymbolsAsZip(ShaderFormat);
+			if (bSymbolsEnabled && (bExportAsZip || FParse::Param(FCommandLine::Get(), TEXT("ShaderSymbolsExportZip"))))
 			{
-				uint32 MultiprocessId = 0;
-				FParse::Value(FCommandLine::Get(), TEXT("-MultiprocessId="), MultiprocessId);
-				FString LeafName;
-				bMultiprocessOwner = MultiprocessId == 0;
-				if (bMultiprocessOwner)
-				{
-					DeleteExistingShaderZips(PlatformFile, ExportPath);
-					LeafName = FString::Printf(TEXT("%s%s"), ZipFileBaseLeafName, ZipFileExtension);
-				}
-				else
-				{
-					LeafName = FString::Printf(TEXT("%s_%d%s"), ZipFileBaseLeafName, MultiprocessId, ZipFileExtension);
-				}
+				FString LeafName = CreateNameAndDeleteOld(MultiprocessId, PlatformFile, ExportPath, ZipFileBaseLeafName, ZipFileExtension);
 				FString SingleFilePath = ExportPath / LeafName;
 
 				IFileHandle* OutputZipFile = PlatformFile.OpenWrite(*SingleFilePath);
@@ -102,6 +112,13 @@ void FShaderSymbolExport::Initialize()
 					ZipWriter = MakeUnique<FZipArchiveWriter>(OutputZipFile);
 				}
 			}
+			
+			if (bSymbolsInfoEnabled)
+			{
+				// if we are exporting collated shader pdb info into one file
+				FString LeafName = CreateNameAndDeleteOld(MultiprocessId, PlatformFile, ExportPath, ZipFileBaseLeafName, InfoFileExtension);
+				InfoFilePath = ExportPath / LeafName;
+			}
 		}
 	}
 
@@ -115,7 +132,7 @@ void FShaderSymbolExport::Initialize()
 	}
 }
 
-void FShaderSymbolExport::WriteSymbolData(const FString& Filename, TConstArrayView<uint8> Contents)
+void FShaderSymbolExport::WriteSymbolData(const FString& Filename, const FString& DebugData, TConstArrayView<uint8> Contents)
 {
 	// Skip this symbol data if we've already exported it before.
 	bool bAlreadyInSet = false;
@@ -138,35 +155,44 @@ void FShaderSymbolExport::WriteSymbolData(const FString& Filename, TConstArrayVi
 		LastReport = TotalSymbolDataBytes;
 	}
 
-	if (ZipWriter)
+	if (ShouldGenerateShaderSymbolsInfo(ShaderFormat) && !DebugData.IsEmpty())
 	{
-		// Append the platform data to the zip file
-		ZipWriter->AddFile(Filename, Contents, FDateTime::Now());
+		// Collect the simple shader symbol information
+		ShaderInfos.Add({ Filename, DebugData });
 	}
-	else
-	{
-		IPlatformFile& PlatformFile = FPlatformFileManager::Get().GetPlatformFile();
 
-		// Write the symbols to the export directory
-		const FString OutputPath = ExportPath / Filename;
-		const FString Directory = FPaths::GetPath(OutputPath);
-
-		// Filename could contain extra folders, so we need to make sure they exist first.
-		if (!PlatformFile.CreateDirectoryTree(*Directory))
+	if (ShouldWriteShaderSymbols(ShaderFormat) && Contents.Num())
+	{
+		if (ZipWriter)
 		{
-			UE_LOG(LogShaderSymbolExport, Error, TEXT("Failed to create shader symbol directory \"%s\"."), *Directory);
+			// Append the platform data to the zip file
+			ZipWriter->AddFile(Filename, Contents, FDateTime::Now());
 		}
 		else
 		{
-			IFileHandle* File = PlatformFile.OpenWrite(*OutputPath);
-			if (!File || !File->Write(Contents.GetData(), Contents.Num()))
+			IPlatformFile& PlatformFile = FPlatformFileManager::Get().GetPlatformFile();
+
+			// Write the symbols to the export directory
+			const FString OutputPath = ExportPath / Filename;
+			const FString Directory = FPaths::GetPath(OutputPath);
+
+			// Filename could contain extra folders, so we need to make sure they exist first.
+			if (!PlatformFile.CreateDirectoryTree(*Directory))
 			{
-				UE_LOG(LogShaderSymbolExport, Error, TEXT("Failed to export shader symbols \"%s\"."), *OutputPath);
+				UE_LOG(LogShaderSymbolExport, Error, TEXT("Failed to create shader symbol directory \"%s\"."), *Directory);
 			}
-
-			if (File)
+			else
 			{
-				delete File;
+				IFileHandle* File = PlatformFile.OpenWrite(*OutputPath);
+				if (!File || !File->Write(Contents.GetData(), Contents.Num()))
+				{
+					UE_LOG(LogShaderSymbolExport, Error, TEXT("Failed to export shader symbols \"%s\"."), *OutputPath);
+				}
+
+				if (File)
+				{
+					delete File;
+				}
 			}
 		}
 	}
@@ -174,7 +200,61 @@ void FShaderSymbolExport::WriteSymbolData(const FString& Filename, TConstArrayVi
 
 void FShaderSymbolExport::NotifyShaderCompilersShutdown()
 {
-	if (bMultiprocessOwner && ZipWriter)
+	if (ShaderInfos.Num())
+	{
+		// sort and combine the data, to either write to a zip or just as a file
+		ShaderInfos.StableSort([](const FShaderInfo& A, const FShaderInfo& B) { return A.Hash < B.Hash; });
+
+		TArray<uint8> Output;
+		for (FShaderInfo Info : ShaderInfos)
+		{
+			auto TmpHash = StringCast<ANSICHAR>(*Info.Hash);
+			auto TmpData = StringCast<ANSICHAR>(*Info.Data);
+			Output.Append((const uint8*)TmpHash.Get(), TmpHash.Length());
+			Output.Add(' ');
+			Output.Append((const uint8*)TmpData.Get(), TmpData.Length());
+			Output.Add('\n');
+		}
+
+		if (InfoFilePath.Len())
+		{
+			IFileManager& FileManager = IFileManager::Get();
+			if (bMultiprocessOwner)
+			{
+				// if we are the multiprocess owner merge in any other files we find
+				IPlatformFile& PlatformFile = FPlatformFileManager::Get().GetPlatformFile();
+				TArray<FString> FilesToMergeIn;
+				PlatformFile.FindFiles(FilesToMergeIn, *ExportPath, InfoFileExtension);
+				for (const FString& InfoFile : FilesToMergeIn)
+				{
+					TUniquePtr<FArchive> Reader = TUniquePtr<FArchive>(FileManager.CreateFileReader(*InfoFile));
+					if (Reader.IsValid())
+					{
+						int64 Size = Reader->TotalSize();
+						TArray<uint8> RawData;
+						RawData.AddUninitialized(Size);
+						Reader->Serialize(RawData.GetData(), Size);
+						Reader->Close();
+						Output.Append(RawData);
+					}
+					PlatformFile.DeleteFile(*InfoFile);
+				}
+			}
+
+			TUniquePtr<FArchive> Writer = TUniquePtr<FArchive>(FileManager.CreateFileWriter(*InfoFilePath));
+			if (Writer.IsValid())
+			{
+				Writer->Serialize(Output.GetData(), Output.Num());
+				Writer->Close();
+			}
+			else
+			{
+				UE_LOG(LogShaderSymbolExport, Error, TEXT("Failed to create shader symbols output file \"%s\"."), *InfoFilePath);
+			}
+		}
+	}
+
+	if (ZipWriter && bMultiprocessOwner)
 	{
 		IPlatformFile& PlatformFile = FPlatformFileManager::Get().GetPlatformFile();
 		TArray<FString> ZipsToMergeIn;
diff --git a/Engine/Source/Developer/ShaderCompilerCommon/Public/ShaderSymbolExport.h b/Engine/Source/Developer/ShaderCompilerCommon/Public/ShaderSymbolExport.h
--- a/Engine/Source/Developer/ShaderCompilerCommon/Public/ShaderSymbolExport.h
+++ b/Engine/Source/Developer/ShaderCompilerCommon/Public/ShaderSymbolExport.h
@@ -32,18 +32,26 @@ public:
 
 private:
 	void Initialize();
-	void WriteSymbolData(const FString& Filename, TConstArrayView<uint8> Contents);
+	void WriteSymbolData(const FString& Filename, const FString& DebugInfo, TConstArrayView<uint8> Contents);
 
 	const FName ShaderFormat;
 
 	TUniquePtr<FZipArchiveWriter> ZipWriter;
 	TSet<FString> ExportedShaders;
 	FString ExportPath;
+	FString InfoFilePath;
 
 	uint64 TotalSymbolDataBytes{ 0 };
 	uint64 TotalSymbolData{ 0 };
 	bool bExportShaderSymbols{ false };
 
+	struct FShaderInfo
+	{
+		FString Hash;
+		FString Data;
+	};
+	TArray<FShaderInfo> ShaderInfos;
+
 	/**
 	 * If true, the current process is the first process in a multiprocess group, or is not in a group,
 	 * and should combine artifacts produced by the other processes. Will also be false if no combination
@@ -73,12 +81,11 @@ inline void FShaderSymbolExport::NotifyShaderCompiled(const TConstArrayView<uint
 
 		for (const auto& SymbolData : FullSymbolData.GetAllSymbolData())
 		{
-			if (TConstArrayView<uint8> Contents = SymbolData.GetContents(); !Contents.IsEmpty())
-			{
-				const FString FileName = SymbolData.GetFilename();
+			const FString FileName = SymbolData.GetFilename();
+			const FString DebugInfo = SymbolData.GetDebugInfo();
+			TConstArrayView<uint8> Contents = SymbolData.GetContents();
 
-				WriteSymbolData(FileName, Contents);
-			}
+			WriteSymbolData(FileName, DebugInfo, Contents);
 		}
 	}
 }
diff --git a/Engine/Source/Runtime/Engine/Private/ShaderCompiler/ShaderCompiler.cpp b/Engine/Source/Runtime/Engine/Private/ShaderCompiler/ShaderCompiler.cpp
--- a/Engine/Source/Runtime/Engine/Private/ShaderCompiler/ShaderCompiler.cpp
+++ b/Engine/Source/Runtime/Engine/Private/ShaderCompiler/ShaderCompiler.cpp
@@ -1654,6 +1654,14 @@ static TAutoConsoleVariable<int32> CVarShadersSymbols(
 	TEXT("This setting can be overriden in any Engine.ini under the [ShaderCompiler] section."),
 	ECVF_ReadOnly);
 
+static TAutoConsoleVariable<int32> CVarShadersSymbolsInfo(
+	TEXT("r.Shaders.SymbolsInfo"),
+	0,
+	TEXT("In lieu of a full set of platform shader PDBs, save out a slimmer ShaderSymbols.Info which contains shader platform hashes and shader debug info.\n")
+	TEXT("An option for when it is not practical to save PDBs for shaders all the time.\n")
+	TEXT("This setting can be overriden in any Engine.ini under the [ShaderCompiler] section."),
+	ECVF_ReadOnly);
+
 static TAutoConsoleVariable<int32> CVarShadersGenerateSymbols(
 	TEXT("r.Shaders.GenerateSymbols"),
 	0,
@@ -6847,6 +6855,10 @@ void GlobalBeginCompileShader(
 	{
 		Input.Environment.CompilerFlags.Add(CFLAG_GenerateSymbols);
 	}
+	if (ShouldGenerateShaderSymbolsInfo(ShaderFormatName))
+	{
+		Input.Environment.CompilerFlags.Add(CFLAG_GenerateSymbolsInfo);
+	}
 
 	// Are symbols based on source or results
 	if (ShouldAllowUniqueShaderSymbols(ShaderFormatName))
diff --git a/Engine/Source/Runtime/RenderCore/Private/Shader.cpp b/Engine/Source/Runtime/RenderCore/Private/Shader.cpp
--- a/Engine/Source/Runtime/RenderCore/Private/Shader.cpp
+++ b/Engine/Source/Runtime/RenderCore/Private/Shader.cpp
@@ -1597,8 +1597,9 @@ void ShaderMapAppendKeyString(EShaderPlatform Platform, FString& KeyString)
 	{
 		// Extra data (names, etc)
 		KeyString += ShouldEnableExtraShaderData(ShaderFormatName) ? TEXT("_ExtraData") : TEXT("");
-		// Symbols
+		// Symbols and/or SymbolsInfo
 		KeyString += ShouldGenerateShaderSymbols(ShaderFormatName) ? TEXT("_Symbols") : TEXT("");
+		KeyString += ShouldGenerateShaderSymbolsInfo(ShaderFormatName) ? TEXT("_SymbolsInfo") : TEXT("");
 		// Are symbols based on source or results
 		KeyString += ShouldAllowUniqueShaderSymbols(ShaderFormatName) ? TEXT("_FullDbg") : TEXT("");
 	}
diff --git a/Engine/Source/Runtime/RenderCore/Private/ShaderCore.cpp b/Engine/Source/Runtime/RenderCore/Private/ShaderCore.cpp
--- a/Engine/Source/Runtime/RenderCore/Private/ShaderCore.cpp
+++ b/Engine/Source/Runtime/RenderCore/Private/ShaderCore.cpp
@@ -871,6 +871,12 @@ bool ShouldGenerateShaderSymbols(FName ShaderFormat)
 	return Symbols.IsEnabled(ShaderFormat) || GenerateSymbols.IsEnabled(ShaderFormat);
 }
 
+bool ShouldGenerateShaderSymbolsInfo(FName ShaderFormat)
+{
+	static const FShaderSymbolSettingHelper SymbolsInfo(TEXT("r.Shaders.SymbolsInfo"));
+	return SymbolsInfo.IsEnabled(ShaderFormat);
+}
+
 bool ShouldWriteShaderSymbols(FName ShaderFormat)
 {
 	static const FShaderSymbolSettingHelper Symbols(TEXT("r.Shaders.Symbols"));
diff --git a/Engine/Source/Runtime/RenderCore/Public/ShaderCompilerCore.h b/Engine/Source/Runtime/RenderCore/Public/ShaderCompilerCore.h
--- a/Engine/Source/Runtime/RenderCore/Public/ShaderCompilerCore.h
+++ b/Engine/Source/Runtime/RenderCore/Public/ShaderCompilerCore.h
@@ -44,10 +44,13 @@ namespace UE::ShaderCompiler
 /** Returns true if shader symbols should be kept for a given platform. */
 extern RENDERCORE_API bool ShouldGenerateShaderSymbols(FName ShaderFormat);
 
+/** Returns true if shader symbol minimal info files should be generated for a given platform. */
+extern RENDERCORE_API bool ShouldGenerateShaderSymbolsInfo(FName ShaderFormat);
+
 /** Returns true if shader symbols should be exported to separate files for a given platform. */
 extern RENDERCORE_API bool ShouldWriteShaderSymbols(FName ShaderFormat);
 
-/** Returns true if the shader symbol path is overriden and OutPathOverride contains the override path. */
+/** Returns true if the shader symbol path is overridden and OutPathOverride contains the override path. */
 extern RENDERCORE_API bool GetShaderSymbolPathOverride(FString& OutPathOverride, FName ShaderFormat);
 
 /** Returns true if (external) shader symbols should be specific to each shader rather than be de-duplicated. */
@@ -81,6 +84,8 @@ enum ECompilerFlags
 	CFLAG_ForceOptimization,
 	// Shader should generate symbols for debugging.
 	CFLAG_GenerateSymbols,
+	// Shader should generate minimal symbols info
+	CFLAG_GenerateSymbolsInfo,
 	// Shader should insert debug/name info at the risk of generating non-deterministic libraries
 	CFLAG_ExtraShaderData,
 	// Allows the (external) symbols to be specific to each shader rather than trying to deduplicate.
@@ -326,6 +331,13 @@ struct FShaderCompilerInput
 		return FPaths::GetCleanFilename(VirtualSourceFilePath);
 	}
 
+	// Common code to generate a debug string to associate with platform-specific shader symbol files and hashes
+	// Currently uses DebugGroupName, but can be updated to contain other important information as needed
+	FString GenerateDebugInfo() const
+	{
+		return DebugGroupName;
+	}
+
 	void GatherSharedInputs(
 		TMap<FString, FString>& ExternalIncludes,
 		TArray<TRefCountPtr<FSharedShaderCompilerEnvironment>>& SharedEnvironments,
