Change D3D diagnostic buffer to output readable assert message (file/line/condition/message).

Use with check() / checkf()

diff --git a/Engine/Shaders/Public/Platform/D3D/D3DCommon.ush b/Engine/Shaders/Public/Platform/D3D/D3DCommon.ush
--- a/Engine/Shaders/Public/Platform/D3D/D3DCommon.ush
+++ b/Engine/Shaders/Public/Platform/D3D/D3DCommon.ush
@@ -76,6 +76,9 @@ void UEReportAssertWithPayload(uint AssertID, uint4 Payload)
 #define PLATFORM_ASSERT3(condition, assert_id, a, b, c)    { if (!(condition)) UEReportAssertWithPayload(assert_id, uint4(a, b, c, 0));  }
 #define PLATFORM_ASSERT4(condition, assert_id, a, b, c, d) { if (!(condition)) UEReportAssertWithPayload(assert_id, uint4(a, b, c, d));  }
 
+#define check(condition) 									{ if (!(condition)) { UEReportAssertWithPayload(0xCECC, uint4(__LINE__, TEXT(__FILE__), TEXT(#condition), 0)); } }
+#define checkf(condition, message) 							{ if (!(condition)) { UEReportAssertWithPayload(0xCECC, uint4(__LINE__, TEXT(__FILE__), message, 0)); } }
+
 #endif // PLATFORM_SUPPORTS_DIAGNOSTIC_BUFFER
 
 #if PLATFORM_SUPPORTS_MESH_SHADERS_TIER0
diff --git a/Engine/Shaders/Public/ShaderVersion.ush b/Engine/Shaders/Public/ShaderVersion.ush
--- a/Engine/Shaders/Public/ShaderVersion.ush
+++ b/Engine/Shaders/Public/ShaderVersion.ush
@@ -2,4 +2,4 @@
 // Each time the console command is executed it generates a new GUID. As this file is included
 // in Platform.ush (which should be included in any shader) it allows to invalidate the shader DDC.
 // 
-// GUID = 111D4D59-4A68-0336-ACC5-4F9E3346A24F
+// GUID = F8F483DB4AC710E5DAF54981165B2A4F
\ No newline at end of file
diff --git a/Engine/Source/Developer/ShaderCompilerCommon/Private/ShaderCompilerCommon.cpp b/Engine/Source/Developer/ShaderCompilerCommon/Private/ShaderCompilerCommon.cpp
--- a/Engine/Source/Developer/ShaderCompilerCommon/Private/ShaderCompilerCommon.cpp
+++ b/Engine/Source/Developer/ShaderCompilerCommon/Private/ShaderCompilerCommon.cpp
@@ -937,7 +937,7 @@ void RemoveUniformBuffersFromSource(const FShaderCompilerEnvironment& Environmen
 ///////////////////////////////////////////////////////////////////////////////////////////////////////////////////////////////////////
 // Process TEXT() macro to convert them into GPU ASCII characters
 
-FString ParseText(const TCHAR* StartPtr, const TCHAR*& EndPtr)
+static FString ParseText(const TCHAR* StartPtr, const TCHAR*& EndPtr)
 {
 	const TCHAR* OpeningBracePtr = FCString::Strstr(StartPtr, TEXT("("));
 	check(OpeningBracePtr);
@@ -959,7 +959,7 @@ FString ParseText(const TCHAR* StartPtr, const TCHAR*& EndPtr)
 	return Out;
 }
 
-void ConvertTextToAsciiCharacter(const FString& InText, FString& OutText, FString& OutEncodedText)
+static void ConvertTextToAsciiCharacter(const FString& InText, FString& OutText, FString& OutEncodedText)
 {
 	const uint32 CharCount = InText.Len();
 	OutEncodedText.Reserve(CharCount * 3); // ~2 digits per character + a comma
@@ -975,8 +975,65 @@ void ConvertTextToAsciiCharacter(const FString& InText, FString& OutText, FStrin
 	}
 }
 
+struct FAssertParsingStack
+{
+	bool IsAssert() const { return BeginPtr != nullptr && EndPtr != nullptr; }
+	const TCHAR* BeginPtr = nullptr;
+	const TCHAR* EndPtr = nullptr;
+};
+
+static bool SearchText(const TCHAR*& InOut, FAssertParsingStack& OutAssertStack)
+{
+	const TCHAR* TextIdentifier = TEXT("TEXT(");
+	const TCHAR* AssertIdentifier = TEXT("UEReportAssertWithPayload(");
+
+	// Find the next TEXT() or UEReportAssertWithPayload()
+	const TCHAR* PrintfPtr = FCString::Strstr(InOut, TextIdentifier);
+	const TCHAR* AssertPtr = FCString::Strstr(InOut, AssertIdentifier);
+
+	// 1. Default is current TEXT parsing
+	InOut = PrintfPtr;
+
+	// 2. If current pointer is beyond the current assert context, reset the context
+	if (InOut >= OutAssertStack.EndPtr)
+	{
+		OutAssertStack.BeginPtr = nullptr;
+		OutAssertStack.EndPtr = nullptr;
+	}
+
+	// 3. If we are within an assert context, continue to part TEXT within that context
+	if (OutAssertStack.IsAssert())
+	{
+		// Nothing to do. InOut will be initialized by PrintfPtr, which is the next TEXT block.
+	}
+	// 4. If a new assert is detected, start a new assert context.
+	else if (AssertPtr && AssertPtr < PrintfPtr)
+	{
+		// Sanity check
+		check(!OutAssertStack.IsAssert());
+
+		// Check if the current assert is valid, i.e., containt a TEXT() argument
+		const TCHAR* EndPtr = nullptr;
+		const FString Tmp = ParseText(AssertPtr, EndPtr);
+		const bool bIsValid = FCString::Strstr(&Tmp[0], TextIdentifier) != nullptr;
+		if (bIsValid)
+		{
+			OutAssertStack.BeginPtr = AssertPtr;
+			OutAssertStack.EndPtr   = EndPtr;
+			InOut = FCString::Strstr(AssertPtr, TextIdentifier);
+		}
+		else
+		{
+			const uint32 LenAssertIdentifier = FString(AssertIdentifier).Len();
+			AssertPtr += LenAssertIdentifier;
+			return SearchText(AssertPtr, OutAssertStack);
+		}
+	}
+	return InOut != nullptr;
+}
+
 // Simple token matching and expansion to replace TEXT macro into supported character string
-void TransformStringIntoCharacterArray(FString& PreprocessedShaderSource)
+void TransformStringIntoCharacterArray(FString& PreprocessedShaderSource, TArray<FShaderDiagnosticData>* OutDiagnosticDatas)
 {
 	// Early out if input is empty; '&PreprocessedShaderSource[0]' below does not return a valid pointer for empty FString
 	if (PreprocessedShaderSource.IsEmpty())
@@ -989,6 +1046,7 @@ void TransformStringIntoCharacterArray(FString& PreprocessedShaderSource)
 		uint32  Index;
 		uint32  Hash;
 		uint32  Offset;
+		bool    bIsAssert;
 		FString SourceText;
 		FString ConvertedText;
 		FString EncodedText;
@@ -999,13 +1057,15 @@ void TransformStringIntoCharacterArray(FString& PreprocessedShaderSource)
 	// 2. Add a text entry
 	// 3. Replace TEXT by its entry number
 	uint32 GlobalCount = 0;
+	uint32 AssertCount = 0;
+	uint32 PrintfCount = 0;
 	{
 		const FString InitHashBegin(TEXT("InitShaderPrintText("));
 		const FString InitHashEnd(TEXT(")"));
 
-		const TCHAR* TextIdentifier = TEXT("TEXT(");
-		const TCHAR* SearchPtr = FCString::Strstr(&PreprocessedShaderSource[0], TextIdentifier);
-		while (SearchPtr)
+		FAssertParsingStack AssertParsingStack;
+		const TCHAR* SearchPtr = &PreprocessedShaderSource[0];
+		while (SearchText(SearchPtr, AssertParsingStack))
 		{
 			const TCHAR* EndPtr = nullptr;
 			FString Text = ParseText(SearchPtr, EndPtr);
@@ -1024,6 +1084,15 @@ void TransformStringIntoCharacterArray(FString& PreprocessedShaderSource)
 				Entry.SourceText	= Text;
 				ConvertTextToAsciiCharacter(Entry.SourceText, Entry.ConvertedText, Entry.EncodedText);
 				Entry.Hash			= CityHash32((const char*)Entry.SourceText.GetCharArray().GetData(), sizeof(FString::ElementType) * Entry.SourceText.Len());
+				Entry.bIsAssert		= AssertParsingStack.IsAssert();
+				if (Entry.bIsAssert)
+				{
+					++AssertCount;
+				}
+				else
+				{
+					++PrintfCount;
+				}
 
 				// Sanity check
 				uint32 HCheck = CityHash32((const char*)Entry.SourceText.GetCharArray().GetData(), sizeof(FString::ElementType) * Entry.SourceText.Len());
@@ -1037,13 +1106,22 @@ void TransformStringIntoCharacterArray(FString& PreprocessedShaderSource)
 				const uint32 CharCount = (EndPtr - SearchPtr) + 1;
 				PreprocessedShaderSource.RemoveAt(StartIndex, CharCount);
 
-				const FString HashText = InitHashBegin + FString::FromInt(EntryIndex) + InitHashEnd;
-				PreprocessedShaderSource.InsertAt(StartIndex, HashText);
+				// * If this is an assert, replace the text with a simple hash
+				// * If this is a regular print, replace the text with IniShaderPrintText for runtime lookup
+				if (Entry.bIsAssert)
+				{
+					const FString HashString = FString::Printf(TEXT("%u"), Entry.Hash);
+					PreprocessedShaderSource.InsertAt(StartIndex, HashString);
+				}
+				else
+				{
+					const FString HashText = InitHashBegin + FString::FromInt(EntryIndex) + InitHashEnd;
+					PreprocessedShaderSource.InsertAt(StartIndex, HashText);
+				}
 
 				// Update SearchPtr, as PreprocessedShaderSource has been modified, and its memory could have been reallocated, causing SearchPtr to be invalid.
 				SearchPtr = &PreprocessedShaderSource[0] + StartIndex;
 			}
-			SearchPtr = FCString::Strstr(SearchPtr, TextIdentifier);
 		}
 	}
 
@@ -1051,7 +1129,7 @@ void TransformStringIntoCharacterArray(FString& PreprocessedShaderSource)
 	// 5. Write the function for fetching character for a given entry index
 	const uint32 EntryCount = Entries.Num();
 	FString TextChars;
-	if (EntryCount>0 && GlobalCount>0)
+	if (PrintfCount>0 && EntryCount>0 && GlobalCount>0)
 	{
 		// 1. Encoded character for each text entry within a single global char array
 		TextChars = FString::Printf(TEXT("static const uint TEXT_CHARS[%d] = {\n"), GlobalCount);
@@ -1104,6 +1182,21 @@ void TransformStringIntoCharacterArray(FString& PreprocessedShaderSource)
 			PreprocessedShaderSource.InsertAt(StartIndex, TextChars);
 		}
 	}
+
+	// 7. Insert assert data into shader compilation output for runtime CPU lookup
+	if (OutDiagnosticDatas && AssertCount > 0)
+	{
+		OutDiagnosticDatas->Reserve(OutDiagnosticDatas->Num() + AssertCount);
+		for (const FTextEntry& E : Entries)
+		{
+			if (E.bIsAssert)
+			{
+				FShaderDiagnosticData& Data = OutDiagnosticDatas->AddDefaulted_GetRef();
+				Data.Hash = E.Hash;
+				Data.Message = E.SourceText;
+			}
+		}
+	}
 }
 
 ///////////////////////////////////////////////////////////////////////////////////////////////////////////////////////////////////////
diff --git a/Engine/Source/Developer/ShaderCompilerCommon/Public/ShaderCompilerCommon.h b/Engine/Source/Developer/ShaderCompilerCommon/Public/ShaderCompilerCommon.h
--- a/Engine/Source/Developer/ShaderCompilerCommon/Public/ShaderCompilerCommon.h
+++ b/Engine/Source/Developer/ShaderCompilerCommon/Public/ShaderCompilerCommon.h
@@ -311,7 +311,7 @@ extern SHADERCOMPILERCOMMON_API void AddUnboundShaderParameterError(
 extern SHADERCOMPILERCOMMON_API void RemoveUniformBuffersFromSource(const FShaderCompilerEnvironment& Environment, FString& PreprocessedShaderSource);
 
 // Processes TEXT macros
-extern SHADERCOMPILERCOMMON_API void TransformStringIntoCharacterArray(FString& PreprocessedShaderSource);
+extern SHADERCOMPILERCOMMON_API void TransformStringIntoCharacterArray(FString& PreprocessedShaderSource, TArray<FShaderDiagnosticData>* OutDiagnosticDatas=nullptr);
 
 extern SHADERCOMPILERCOMMON_API bool RemoveUnusedOutputs(FString& InOutSourceCode, const TArray<FString>& InUsedOutputs, const TArray<FString>& InExceptions, FString& InOutEntryPoint, TArray<FString>& OutErrors);
 
diff --git a/Engine/Source/Developer/Windows/ShaderFormatD3D/Private/D3DShaderCompiler.cpp b/Engine/Source/Developer/Windows/ShaderFormatD3D/Private/D3DShaderCompiler.cpp
--- a/Engine/Source/Developer/Windows/ShaderFormatD3D/Private/D3DShaderCompiler.cpp
+++ b/Engine/Source/Developer/Windows/ShaderFormatD3D/Private/D3DShaderCompiler.cpp
@@ -1198,7 +1198,7 @@ bool PreprocessD3DShader(
 	}
 
 	// Process TEXT macro.
-	TransformStringIntoCharacterArray(PreprocessedSource);
+	TransformStringIntoCharacterArray(PreprocessedSource, &Output.EditDiagnosticDatas());
 
 	TArray<FString> FilteredErrors;
 	// Run the shader minifier
@@ -1294,4 +1294,6 @@ void CompileD3DShader(const FShaderCompilerInput& Input, const FShaderPreprocess
 			FPlatformMisc::LowLevelOutputDebugStringf(TEXT("%s\n"), *Error.GetErrorStringWithLineMarker());
 		}
 	}
+
+	Output.ShaderDiagnosticDatas = PreprocessOutput.GetDiagnosticDatas();
 }
diff --git a/Engine/Source/Developer/Windows/ShaderFormatD3D/Private/D3DShaderCompiler.inl b/Engine/Source/Developer/Windows/ShaderFormatD3D/Private/D3DShaderCompiler.inl
--- a/Engine/Source/Developer/Windows/ShaderFormatD3D/Private/D3DShaderCompiler.inl
+++ b/Engine/Source/Developer/Windows/ShaderFormatD3D/Private/D3DShaderCompiler.inl
@@ -365,6 +365,7 @@ inline void GenerateFinalOutput(TRefCountPtr<TBlob>& CompressedData,
 	}
 
 	Output.SerializeShaderCodeValidation();
+	Output.SerializeShaderDiagnosticData();
 
 	// Set the number of instructions.
 	Output.NumInstructions = NumInstructions;
diff --git a/Engine/Source/Runtime/D3D12RHI/Private/D3D12Shaders.cpp b/Engine/Source/Runtime/D3D12RHI/Private/D3D12Shaders.cpp
--- a/Engine/Source/Runtime/D3D12RHI/Private/D3D12Shaders.cpp
+++ b/Engine/Source/Runtime/D3D12RHI/Private/D3D12Shaders.cpp
@@ -53,6 +53,7 @@ static inline bool ReadShaderOptionalData(FShaderCodeReader& InShaderCode, TShad
 #endif
 
 	UE::RHICore::SetupShaderCodeValidationData(&OutShader, InShaderCode);
+	UE::RHICore::SetupShaderDiagnosticData(&OutShader, InShaderCode);
 
 	return true;
 }
diff --git a/Engine/Source/Runtime/D3D12RHI/Private/D3D12Util.cpp b/Engine/Source/Runtime/D3D12RHI/Private/D3D12Util.cpp
--- a/Engine/Source/Runtime/D3D12RHI/Private/D3D12Util.cpp
+++ b/Engine/Source/Runtime/D3D12RHI/Private/D3D12Util.cpp
@@ -249,7 +249,17 @@ static bool LogBreadcrumbData(D3D12RHI::FD3DGPUProfiler& GPUProfiler, FD3D12Queu
 	const FD3D12DiagnosticBufferData* DiagnosticData = Queue.GetDiagnosticBufferData();
 	if (DiagnosticData && DiagnosticData->Counter)
 	{
-		UE_LOG(LogD3D12RHI, Error, TEXT("[GPUBreadCrumb]\t\tShader assertion failed! ID: 0x%08X (%d)"), DiagnosticData->MessageID, DiagnosticData->MessageID);
+		const uint32 Line = DiagnosticData->Payload.AsUint[0];
+		const FString* File = UE::RHICore::GetDiagnosticMessage(DiagnosticData->Payload.AsUint[1]);
+		const FString* Message = UE::RHICore::GetDiagnosticMessage(DiagnosticData->Payload.AsUint[2]);
+		if (File && Message)
+		{
+			UE_LOG(LogD3D12RHI, Error, TEXT("[GPUBreadCrumb]\t\tShader assertion failed - %s:%d - %s"), **File, Line, **Message);
+		}
+		else
+		{
+			UE_LOG(LogD3D12RHI, Error, TEXT("[GPUBreadCrumb]\t\tShader assertion failed! ID: 0x%08X (%d)"), DiagnosticData->MessageID, DiagnosticData->MessageID);
+		}
 
 		{
 			const int32* Payload = DiagnosticData->Payload.AsInt;
diff --git a/Engine/Source/Runtime/RHICore/Private/RHICoreShader.cpp b/Engine/Source/Runtime/RHICore/Private/RHICoreShader.cpp
--- a/Engine/Source/Runtime/RHICore/Private/RHICoreShader.cpp
+++ b/Engine/Source/Runtime/RHICore/Private/RHICoreShader.cpp
@@ -112,5 +112,39 @@ void DispatchShaderBundleEmulation(
 	}
 }
 
+const bool GRHIShaderDiagnosticEnabled = true;
+void SetupShaderDiagnosticData(FRHIShader* RHIShader, FShaderCodeReader& ShaderCodeReader)
+{
+	if (RHIShader && GRHIShaderDiagnosticEnabled)
+	{
+		int32 ShaderDiagnosticExtensionSize = 0;
+		const uint8* ShaderDiagnosticExtensionData = ShaderCodeReader.FindOptionalDataAndSize(FShaderDiagnosticExtension::Key, ShaderDiagnosticExtensionSize);
+		if (ShaderDiagnosticExtensionData && ShaderDiagnosticExtensionSize > 0)
+		{
+			FBufferReader ArValidationExtensionData((void*)ShaderDiagnosticExtensionData, ShaderDiagnosticExtensionSize, false);
+			FShaderDiagnosticExtension ShaderDiagnosticExtension;
+			ArValidationExtensionData << ShaderDiagnosticExtension;
+			RegisterDiagnosticMessages(ShaderDiagnosticExtension.ShaderDiagnosticDatas);
+		}
+	}
+}
+
+TArray<FShaderDiagnosticData> GShaderDiagnosticDatas;
+void RegisterDiagnosticMessages(const TArray<FShaderDiagnosticData>& In)
+{
+	// Not thread safe
+	GShaderDiagnosticDatas.Append(In);
+}
+
+const FString* GetDiagnosticMessage(uint32 MessageID)
+{
+	// Not thread safe
+	if (const FShaderDiagnosticData* Found = GShaderDiagnosticDatas.FindByPredicate([MessageID](const FShaderDiagnosticData& In) { return In.Hash == MessageID; }))
+	{
+		return &Found->Message;
+	}
+	return nullptr;
+}
+
 } //! RHICore
 } //! UE
\ No newline at end of file
diff --git a/Engine/Source/Runtime/RHICore/Public/RHICoreShader.h b/Engine/Source/Runtime/RHICore/Public/RHICoreShader.h
--- a/Engine/Source/Runtime/RHICore/Public/RHICoreShader.h
+++ b/Engine/Source/Runtime/RHICore/Public/RHICoreShader.h
@@ -13,6 +13,9 @@ namespace RHICore
 /** Validates that the uniform buffer at the requested static slot. */
 extern RHICORE_API void ValidateStaticUniformBuffer(FRHIUniformBuffer* UniformBuffer, FUniformBufferStaticSlot Slot, uint32 ExpectedHash);
 extern RHICORE_API void SetupShaderCodeValidationData(FRHIShader* RHIShader, class FShaderCodeReader& ShaderCodeReader);
+extern RHICORE_API void SetupShaderDiagnosticData(FRHIShader* RHIShader, class FShaderCodeReader& ShaderCodeReader);
+extern RHICORE_API void RegisterDiagnosticMessages(const TArray<FShaderDiagnosticData>& In);
+extern RHICORE_API const FString* GetDiagnosticMessage(uint32 MessageID);
 
 /** Common implementation of dispatch shader bundle emulation shared by RHIs */
 extern RHICORE_API void DispatchShaderBundleEmulation(
diff --git a/Engine/Source/Runtime/RenderCore/Private/ShaderCore.cpp b/Engine/Source/Runtime/RenderCore/Private/ShaderCore.cpp
--- a/Engine/Source/Runtime/RenderCore/Private/ShaderCore.cpp
+++ b/Engine/Source/Runtime/RenderCore/Private/ShaderCore.cpp
@@ -977,6 +977,20 @@ void FShaderCompilerOutput::SerializeShaderCodeValidation()
 	}
 }
 
+void FShaderCompilerOutput::SerializeShaderDiagnosticData()
+{
+	if (ShaderDiagnosticDatas.Num() > 0)
+	{
+		FShaderDiagnosticExtension ShaderDiagnosticExtension;
+		ShaderDiagnosticExtension.ShaderDiagnosticDatas = ShaderDiagnosticDatas;
+
+		TArray<uint8> WriterBytes;
+		FMemoryWriter Writer(WriterBytes);
+		Writer << ShaderDiagnosticExtension;
+		ShaderCode.AddOptionalData(FShaderDiagnosticExtension::Key, WriterBytes.GetData(), WriterBytes.Num());
+	}
+}
+
 static void ReportVirtualShaderFilePathError(TArray<FShaderCompilerError>* CompileErrors, FString ErrorString)
 {
 	if (CompileErrors)
@@ -1228,6 +1242,7 @@ public:
 
 		Job.PreprocessOutput.ElapsedTime = FPlatformTime::Seconds() - StartPreprocessTime;
 		Job.Output.PreprocessTime = Job.PreprocessOutput.ElapsedTime;
+		Job.Output.ShaderDiagnosticDatas = Job.PreprocessOutput.GetDiagnosticDatas();
 		return Job.PreprocessOutput.bSucceeded;
 	}
 
diff --git a/Engine/Source/Runtime/RenderCore/Private/ShaderPreprocessTypes.cpp b/Engine/Source/Runtime/RenderCore/Private/ShaderPreprocessTypes.cpp
--- a/Engine/Source/Runtime/RenderCore/Private/ShaderPreprocessTypes.cpp
+++ b/Engine/Source/Runtime/RenderCore/Private/ShaderPreprocessTypes.cpp
@@ -295,5 +295,6 @@ FArchive& operator<<(FArchive& Ar, FShaderPreprocessOutput& PreprocessOutput)
 	Ar << PreprocessOutput.PragmaDirectives;
 	Ar << PreprocessOutput.PreprocessedSource;
 	Ar << PreprocessOutput.ParameterParser;
+	Ar << PreprocessOutput.ShaderDiagnosticDatas;
 	return Ar;
 }
diff --git a/Engine/Source/Runtime/RenderCore/Public/ShaderCompilerCore.h b/Engine/Source/Runtime/RenderCore/Public/ShaderCompilerCore.h
--- a/Engine/Source/Runtime/RenderCore/Public/ShaderCompilerCore.h
+++ b/Engine/Source/Runtime/RenderCore/Public/ShaderCompilerCore.h
@@ -589,6 +589,8 @@ struct FShaderCompilerOutput
 
 	TArray<FShaderCodeValidationStride> ParametersStrideToValidate;
 
+	TArray<FShaderDiagnosticData> ShaderDiagnosticDatas;
+
 	/** This field should be set by backends which do not implement the independent preprocessing API to contain the "final" shader source as 
 	 * passed to the platform compiler. For backends that do implement this API this is superceded by ModifiedShaderSource (and will eventually
 	 * be deprecated).
@@ -619,6 +621,9 @@ struct FShaderCompilerOutput
 	/** Add optional data in ShaderCode to perform additional shader input validation at runtime*/
 	RENDERCORE_API void SerializeShaderCodeValidation();
 
+	/** Add optional diagnostic data in ShaderCode to perform assert translation at runtime*/
+	RENDERCORE_API void SerializeShaderDiagnosticData();
+
 	friend FArchive& operator<<(FArchive& Ar, FShaderCompilerOutput& Output)
 	{
 		// Note: this serialize is used to pass between UE and the shader compile worker, recompile both when modifying
diff --git a/Engine/Source/Runtime/RenderCore/Public/ShaderCore.h b/Engine/Source/Runtime/RenderCore/Public/ShaderCore.h
--- a/Engine/Source/Runtime/RenderCore/Public/ShaderCore.h
+++ b/Engine/Source/Runtime/RenderCore/Public/ShaderCore.h
@@ -784,6 +784,34 @@ struct FShaderCodeValidationExtension
 	}
 };
 
+struct FShaderDiagnosticData
+{
+	uint32 Hash;
+	FString Message;
+};
+
+inline FArchive& operator<<(FArchive& Ar, FShaderDiagnosticData& ShaderCodeDiagnosticData)
+{
+	return Ar << ShaderCodeDiagnosticData.Hash << ShaderCodeDiagnosticData.Message;
+}
+
+struct FShaderDiagnosticExtension
+{
+	// for FindOptionalData() and AddOptionalData()
+	static constexpr uint8 Key = 'D';
+	static constexpr uint16 StaticVersion = 0;
+
+	TArray<FShaderDiagnosticData> ShaderDiagnosticDatas;
+	uint16 Version = StaticVersion;
+
+	friend FArchive& operator<<(FArchive& Ar, FShaderDiagnosticExtension& Extension)
+	{
+		Ar << Extension.Version;
+		Ar << Extension.ShaderDiagnosticDatas;
+		return Ar;
+	}
+};
+
 #ifndef RENDERCORE_ATTRIBUTE_UNALIGNED
 // TODO find out if using GCC_ALIGN(1) instead of this new #define break on all kinds of platforms...
 #define RENDERCORE_ATTRIBUTE_UNALIGNED
diff --git a/Engine/Source/Runtime/RenderCore/Public/ShaderPreprocessTypes.h b/Engine/Source/Runtime/RenderCore/Public/ShaderPreprocessTypes.h
--- a/Engine/Source/Runtime/RenderCore/Public/ShaderPreprocessTypes.h
+++ b/Engine/Source/Runtime/RenderCore/Public/ShaderPreprocessTypes.h
@@ -152,6 +152,16 @@ public:
 		return ElapsedTime;
 	}
 
+	TArray<FShaderDiagnosticData>& EditDiagnosticDatas()
+	{
+		return ShaderDiagnosticDatas;
+	}
+
+	const TArray<FShaderDiagnosticData>& GetDiagnosticDatas() const
+	{
+		return ShaderDiagnosticDatas;
+	}
+
 	friend FArchive& operator<<(FArchive& Ar, FShaderPreprocessOutput& PreprocessOutput);
 
 private:
@@ -184,4 +194,6 @@ private:
 	bool bIsSecondary = false;
 
 	FShaderParameterParser ParameterParser;
+
+	TArray<FShaderDiagnosticData> ShaderDiagnosticDatas;
 };
