Fix removal of unused shader interpolators for DXC backends: Keep original semantic name including its index and separate it from its overridden binding slot.

#rnx

diff --git a/Engine/Source/Developer/ShaderCompilerCommon/Private/HlslccHeaderWriter.cpp b/Engine/Source/Developer/ShaderCompilerCommon/Private/HlslccHeaderWriter.cpp
--- a/Engine/Source/Developer/ShaderCompilerCommon/Private/HlslccHeaderWriter.cpp
+++ b/Engine/Source/Developer/ShaderCompilerCommon/Private/HlslccHeaderWriter.cpp
@@ -280,17 +280,10 @@ namespace CrossCompiler
 		}
 
 		const FString TypeSpecifier = ConvertMetaDataTypeSpecifier(*Attribute.type_description);
-		FString Semantic = ConvertAttributeToMetaDataSemantic(Attribute.semantic, Attribute.built_in, bIsInput);
+		const FString Semantic = ConvertAttributeToMetaDataSemantic(Attribute.semantic, Attribute.built_in, bIsInput);
 
 		if (Attribute.array.dims_count > 0)
 		{
-			// Get semantic without index, e.g. "out_Target0" -> "out_Target"
-			const int32 SemanticIndexPos = FindIndexInHlslSemantic(Semantic);
-			if (SemanticIndexPos != INDEX_NONE)
-			{
-				Semantic = Semantic.Left(SemanticIndexPos);
-			}
-
 			if (Attribute.location == -1)
 			{
 				// Flatten array dimensions, e.g. from float4[3][2] -> float4[6]
@@ -300,11 +293,12 @@ namespace CrossCompiler
 				for (uint32 FlattenedArrayIndex = 0; FlattenedArrayIndex < FlattenedArrayDim; ++FlattenedArrayIndex)
 				{
 					// If there is no binding slot, emit output as system value array such as "gl_SampleMask[]"
+					const uint32 BindingSlot = Attribute.location;
 					MetaDataPrintf(
 						OutMetaData,
 						TEXT("%s;%d:%s[%d]"),
 						*TypeSpecifier, // type specifier
-						Attribute.location,
+						BindingSlot,
 						*Semantic,
 						FlattenedArrayIndex
 					);
@@ -322,11 +316,10 @@ namespace CrossCompiler
 					const uint32 BindingSlot = Attribute.location + FlattenedArrayIndex;
 					MetaDataPrintf(
 						OutMetaData,
-						TEXT("%s;%d:%s%d"),
+						TEXT("%s;%d:%s"),
 						*TypeSpecifier, // Type specifier
 						BindingSlot,
-						*Semantic,
-						BindingSlot
+						*Semantic
 					);
 				}
 			}
@@ -341,12 +334,11 @@ namespace CrossCompiler
 					const uint32 BindingSlot = Attribute.location + FlattenedArrayIndex;
 					MetaDataPrintf(
 						OutMetaData,
-						TEXT("%s[%d];%d:%s%d"),
+						TEXT("%s[%d];%d:%s"),
 						*TypeSpecifier, // Type specifier
 						Attribute.array.dims[0], // Outer most array dimension
 						BindingSlot,
-						*Semantic,
-						BindingSlot
+						*Semantic
 					);
 				}
 			}
@@ -355,12 +347,11 @@ namespace CrossCompiler
 				const uint32 BindingSlot = Attribute.location;
 				MetaDataPrintf(
 					OutMetaData,
-					TEXT("%s[%d];%d:%s%d"),
+					TEXT("%s[%d];%d:%s"),
 					*TypeSpecifier, // Type specifier
 					Attribute.array.dims[0], // Outer most array dimension
 					BindingSlot,
-					*Semantic,
-					BindingSlot
+					*Semantic
 				);
 			}
 		}
