VulkanRHI: Add quotes around the entrypoint name in the command line for cases where multiple entrypoints are possible (like RT hit group)

#rnx

diff --git a/Engine/Source/Developer/ShaderCompilerCommon/Private/ShaderCompilerCommon.cpp b/Engine/Source/Developer/ShaderCompilerCommon/Private/ShaderCompilerCommon.cpp
--- a/Engine/Source/Developer/ShaderCompilerCommon/Private/ShaderCompilerCommon.cpp
+++ b/Engine/Source/Developer/ShaderCompilerCommon/Private/ShaderCompilerCommon.cpp
@@ -1205,10 +1205,10 @@ FString CreateShaderCompilerWorkerDirectCommandLine(const FShaderCompilerInput&
 {
 	FString Text(TEXT("-directcompile -format="));
 	Text += Input.ShaderFormat.GetPlainNameString();
-	Text += TEXT(" -entry=");
+	Text += TEXT(" -entry=\"");
 	Text += Input.EntryPointName;
 
-	Text += TEXT(" -shaderPlatformName=");
+	Text += TEXT("\" -shaderPlatformName=");
 	Text += Input.ShaderPlatformName.GetPlainNameString();
 
 	switch (Input.Target.Frequency)
