ShaderConductor:  Change the default for bDisableScalarBlockLayout to TRUE so that it matches the previous behavior where it was never truly enabled.  Attempt to fix build where shader alignment issues started poping up.

#rnx

diff --git a/Engine/Source/Developer/ShaderCompilerCommon/Public/ShaderConductorContext.h b/Engine/Source/Developer/ShaderCompilerCommon/Public/ShaderConductorContext.h
--- a/Engine/Source/Developer/ShaderCompilerCommon/Public/ShaderConductorContext.h
+++ b/Engine/Source/Developer/ShaderCompilerCommon/Public/ShaderConductorContext.h
@@ -81,7 +81,7 @@ namespace CrossCompiler
 		bool bEnableFMAPass = false;
 
 		/** Disables scalar block layout for structured buffers. True for Vulkan mobile due to low coverage of 'VK_EXT_scalar_block_layout' extension. */
-		bool bDisableScalarBlockLayout = false;
+		bool bDisableScalarBlockLayout = true;
 
 		/** Enables separate samplers in GLSL via extensions. */
 		bool bEnableSeparateSamplersInGlsl = false;
