Fix usage of const TCHAR* in the condition of checkf or ensureMsgf and test the strings for nullptr instead

#rnx

diff --git a/Engine/Source/Developer/ShaderCompilerCommon/Private/HlslAST.cpp b/Engine/Source/Developer/ShaderCompilerCommon/Private/HlslAST.cpp
--- a/Engine/Source/Developer/ShaderCompilerCommon/Private/HlslAST.cpp
+++ b/Engine/Source/Developer/ShaderCompilerCommon/Private/HlslAST.cpp
@@ -328,7 +328,7 @@ namespace CrossCompiler
 		{
 			if (IsConstant())
 			{
-				checkf(Identifier, TEXT("Null identifier, literaltype %d"), (int32)LiteralType);
+				checkf(Identifier!=nullptr, TEXT("Null identifier, literaltype %d"), (int32)LiteralType);
 				OutValue = (int32)FCString::Atoi(Identifier);
 				return true;
 			}
diff --git a/Engine/Source/Developer/ShaderCompilerCommon/Private/ShaderConductorContext.cpp b/Engine/Source/Developer/ShaderCompilerCommon/Private/ShaderConductorContext.cpp
--- a/Engine/Source/Developer/ShaderCompilerCommon/Private/ShaderConductorContext.cpp
+++ b/Engine/Source/Developer/ShaderCompilerCommon/Private/ShaderConductorContext.cpp
@@ -176,7 +176,7 @@ namespace CrossCompiler
 	{
 		OutTargetDesc.language = ShaderConductor::ShadingLanguage::Hlsl;
 		OutTargetDesc.version = GetHlslVersionString(InTarget.Version);
-		checkf(OutTargetDesc.version, TEXT("Unsupported target shader version for HLSL: SM%d.%d"), InTarget.Version / 10, InTarget.Version % 10);
+		checkf(OutTargetDesc.version!=nullptr, TEXT("Unsupported target shader version for HLSL: SM%d.%d"), InTarget.Version / 10, InTarget.Version % 10);
 	}
 
 	static const ANSICHAR* GetGlslFamilyVersionString(int32 Version)
@@ -200,7 +200,7 @@ namespace CrossCompiler
 	{
 		OutTargetDesc.language = (InTarget.Language == EShaderConductorLanguage::Glsl ? ShaderConductor::ShadingLanguage::Glsl : ShaderConductor::ShadingLanguage::Essl);
 		OutTargetDesc.version = GetGlslFamilyVersionString(InTarget.Version);
-		checkf(OutTargetDesc.version, TEXT("Unsupported target shader version for GLSL family: %d"), InTarget.Version);
+		checkf(OutTargetDesc.version!=nullptr, TEXT("Unsupported target shader version for GLSL family: %d"), InTarget.Version);
 	}
 
 	static const ANSICHAR* GetMetalFamilyVersionString(int32 Version)
@@ -224,7 +224,7 @@ namespace CrossCompiler
 	{
 		OutTargetDesc.language = (InTarget.Language == EShaderConductorLanguage::Metal_macOS ? ShaderConductor::ShadingLanguage::Msl_macOS : ShaderConductor::ShadingLanguage::Msl_iOS);
 		OutTargetDesc.version = GetMetalFamilyVersionString(InTarget.Version);
-		checkf(OutTargetDesc.version, TEXT("Unsupported target shader version for Metal family: %d"), InTarget.Version);
+		checkf(OutTargetDesc.version!=nullptr, TEXT("Unsupported target shader version for Metal family: %d"), InTarget.Version);
 	}
 
 	// Converts an array of FString to a C-style array of char* pointers
diff --git a/Engine/Source/Runtime/Core/Private/Misc/AES.cpp b/Engine/Source/Runtime/Core/Private/Misc/AES.cpp
--- a/Engine/Source/Runtime/Core/Private/Misc/AES.cpp
+++ b/Engine/Source/Runtime/Core/Private/Misc/AES.cpp
@@ -973,7 +973,7 @@ void FAES::EncryptData(uint8 *Contents, uint64 NumBytes, const FAESKey& Key)
 
 void FAES::EncryptData(uint8 *Contents, uint64 NumBytes, const ANSICHAR* Key)
 {
-	checkf(Key, TEXT("No encryption key specified"));
+	checkf(Key!=nullptr, TEXT("No encryption key specified"));
 	EncryptData(Contents, NumBytes, (const uint8*)Key, TCString<ANSICHAR>::Strlen(Key));
 }
 
@@ -1008,7 +1008,7 @@ void FAES::DecryptData(uint8* Contents, uint64 NumBytes, const FAESKey& Key)
 
 void FAES::DecryptData(uint8 *Contents, uint64 NumBytes, const ANSICHAR* Key)
 {
-	checkf(Key, TEXT("No valid decryption key specified"));
+	checkf(Key!=nullptr, TEXT("No valid decryption key specified"));
 	DecryptData(Contents, NumBytes, (const uint8*)Key, TCString<ANSICHAR>::Strlen(Key));
 }
 
diff --git a/Engine/Source/Runtime/Core/Public/Misc/Crc.h b/Engine/Source/Runtime/Core/Public/Misc/Crc.h
--- a/Engine/Source/Runtime/Core/Public/Misc/Crc.h
+++ b/Engine/Source/Runtime/Core/Public/Misc/Crc.h
@@ -215,7 +215,7 @@ inline uint32 FCrc::Strihash_DEPRECATED(const UTF8CHAR* Data)
 	WIDECHAR* Temp = new WIDECHAR[ConvertedLen];
 
 	WIDECHAR* TempEnd = FPlatformString::Convert(Temp, ConvertedLen, Data, Len);
-	checkf(TempEnd, TEXT("String conversion unsuccessful"));
+	checkf(TempEnd!=nullptr, TEXT("String conversion unsuccessful"));
 
 	// This doesn't work for strings containing characters outside the BMP, but
 	// then neither does the WIDECHAR overload.
@@ -236,7 +236,7 @@ inline uint32 FCrc::Strihash_DEPRECATED(const int32 DataLen, const UTF8CHAR* Dat
 	WIDECHAR* Temp = new WIDECHAR[ConvertedLen];
 
 	WIDECHAR* TempEnd = FPlatformString::Convert(Temp, ConvertedLen, Data, DataLen);
-	checkf(TempEnd, TEXT("String conversion unsuccessful"));
+	checkf(TempEnd!=nullptr, TEXT("String conversion unsuccessful"));
 
 	// This doesn't work for strings containing characters outside the BMP, but
 	// then neither does the WIDECHAR overload.
diff --git a/Engine/Source/Runtime/NNE/Internal/NNEAttributeMap.h b/Engine/Source/Runtime/NNE/Internal/NNEAttributeMap.h
--- a/Engine/Source/Runtime/NNE/Internal/NNEAttributeMap.h
+++ b/Engine/Source/Runtime/NNE/Internal/NNEAttributeMap.h
@@ -13,7 +13,10 @@ public:
 	//Set attribute
 	void SetAttribute(const FString& Name, const FNNEAttributeValue& Value)
 	{
-		checkf(nullptr == Attributes.FindByPredicate([Name](const FEntry& Other) { return Other.Name == Name; }), TEXT("Attribute name should be unique"));
+#if DO_CHECK
+		const bool bIsUnique = (nullptr == Attributes.FindByPredicate([Name](const FEntry& Other) { return Other.Name == Name; }));
+		checkf(bIsUnique, TEXT("Attribute name should be unique"));
+#endif
 		Attributes.Emplace(Name, Value);
 	}
 
diff --git a/Engine/Source/Runtime/RHI/Public/RHIValidationContext.h b/Engine/Source/Runtime/RHI/Public/RHIValidationContext.h
--- a/Engine/Source/Runtime/RHI/Public/RHIValidationContext.h
+++ b/Engine/Source/Runtime/RHI/Public/RHIValidationContext.h
@@ -1114,7 +1114,7 @@ public:
 	virtual void RHIBeginRenderPass(const FRHIRenderPassInfo& InInfo, const TCHAR* InName) override final
 	{
 		checkf(!State.bInsideBeginRenderPass, TEXT("Trying to begin RenderPass '%s', but already inside '%s'!"), *State.RenderPassName, InName);
-		checkf(InName, TEXT("RenderPass should have a name!"));
+		checkf(InName!=nullptr, TEXT("RenderPass should have a name!"));
 		State.bInsideBeginRenderPass = true;
 		State.RenderPassInfo = InInfo;
 		State.RenderPassName = InName;
diff --git a/Engine/Source/Runtime/RenderCore/Private/RenderGraphUtils.cpp b/Engine/Source/Runtime/RenderCore/Private/RenderGraphUtils.cpp
--- a/Engine/Source/Runtime/RenderCore/Private/RenderGraphUtils.cpp
+++ b/Engine/Source/Runtime/RenderCore/Private/RenderGraphUtils.cpp
@@ -1007,7 +1007,7 @@ FRDGBufferRef CreateVertexBuffer(
 	uint64 InitialDataSize,
 	ERDGInitialDataFlags InitialDataFlags)
 {
-	checkf(Name, TEXT("Buffer must have a name."));
+	checkf(Name!=nullptr, TEXT("Buffer must have a name."));
 	checkf(EnumHasAnyFlags(Desc.Usage, EBufferUsageFlags::VertexBuffer), TEXT("CreateVertexBuffer called with an FRDGBufferDesc underlying type that is not 'VertexBuffer'. Buffer: %s"), Name);
 
 	FRDGBufferRef Buffer = GraphBuilder.CreateBuffer(Desc, Name);
diff --git a/Engine/Source/Runtime/RenderCore/Private/RenderGraphValidation.cpp b/Engine/Source/Runtime/RenderCore/Private/RenderGraphValidation.cpp
--- a/Engine/Source/Runtime/RenderCore/Private/RenderGraphValidation.cpp
+++ b/Engine/Source/Runtime/RenderCore/Private/RenderGraphValidation.cpp
@@ -296,7 +296,7 @@ void FRDGUserValidation::ValidateRegisterExternalTexture(
 		return;
 	}
 
-	checkf(Name, TEXT("Attempted to register external texture with NULL name."));
+	checkf(Name!=nullptr, TEXT("Attempted to register external texture with NULL name."));
 	checkf(ExternalPooledTexture.IsValid(), TEXT("Attempted to register NULL external texture."));
 	ExecuteGuard(TEXT("RegisterExternalTexture"), Name);
 }
@@ -308,7 +308,7 @@ void FRDGUserValidation::ValidateRegisterExternalBuffer(const TRefCountPtr<FRDGP
 		return;
 	}
 
-	checkf(Name, TEXT("Attempted to register external buffer with NULL name."));
+	checkf(Name!=nullptr, TEXT("Attempted to register external buffer with NULL name."));
 	checkf(ExternalPooledBuffer.IsValid(), TEXT("Attempted to register NULL external buffer."));
 	ExecuteGuard(TEXT("RegisterExternalBuffer"), Name);
 }
@@ -340,7 +340,7 @@ void FRDGUserValidation::ValidateCreateTexture(const FRDGTextureDesc& Desc, cons
 		return;
 	}
 
-	checkf(Name, TEXT("Creating a texture requires a valid debug name."));
+	checkf(Name!=nullptr, TEXT("Creating a texture requires a valid debug name."));
 	ExecuteGuard(TEXT("CreateTexture"), Name);
 
 	// Make sure the descriptor is supported by the RHI.
@@ -366,7 +366,7 @@ void FRDGUserValidation::ValidateCreateBuffer(const FRDGBufferDesc& Desc, const
 		return;
 	}
 
-	checkf(Name, TEXT("Creating a buffer requires a valid debug name."));
+	checkf(Name!=nullptr, TEXT("Creating a buffer requires a valid debug name."));
 	ExecuteGuard(TEXT("CreateBuffer"), Name);
 
 	if (EnumHasAllFlags(Desc.Usage, EBufferUsageFlags::StructuredBuffer | EBufferUsageFlags::ByteAddressBuffer))
