Added two missing MIMG instructions to RDNA2_ISA.h

diff --git a/Engine/Source/Developer/ShaderCompilerCommon/Private/ISAParser/QuadModeDetect.cpp b/Engine/Source/Developer/ShaderCompilerCommon/Private/ISAParser/QuadModeDetect.cpp
--- a/Engine/Source/Developer/ShaderCompilerCommon/Private/ISAParser/QuadModeDetect.cpp
+++ b/Engine/Source/Developer/ShaderCompilerCommon/Private/ISAParser/QuadModeDetect.cpp
@@ -598,6 +598,8 @@ struct FQuadModeParserRDNA2
 			case RDNA2::EMIMGOps::image_gather4_c_b_cl_o:
 			case RDNA2::EMIMGOps::image_get_lod:
 			case RDNA2::EMIMGOps::image_gather4h:
+			case RDNA2::EMIMGOps::image_gather4h_pck:
+			case RDNA2::EMIMGOps::image_gather8h_pck:
 			{
 				if (bDebug)
 				{
diff --git a/Engine/Source/Developer/ShaderCompilerCommon/Private/ISAParser/RDNA2_ISA.h b/Engine/Source/Developer/ShaderCompilerCommon/Private/ISAParser/RDNA2_ISA.h
--- a/Engine/Source/Developer/ShaderCompilerCommon/Private/ISAParser/RDNA2_ISA.h
+++ b/Engine/Source/Developer/ShaderCompilerCommon/Private/ISAParser/RDNA2_ISA.h
@@ -1659,6 +1659,8 @@ enum class EMIMGOps : uint16
 	image_gather4_c_lz_o = 95,
 	image_get_lod = 96,
 	image_gather4h = 97,
+	image_gather4h_pck = 98,
+	image_gather8h_pck = 99,
 	image_sample_d_g16 = 162,
 	image_sample_d_cl_g16 = 163,
 	image_sample_c_d_g16 = 170,
@@ -1760,6 +1762,8 @@ const char* ToString(EMIMGOps Op)
 		OP_TO_STRING_CASE(image_gather4_c_lz_o);
 		OP_TO_STRING_CASE(image_get_lod);
 		OP_TO_STRING_CASE(image_gather4h);
+		OP_TO_STRING_CASE(image_gather4h_pck);
+		OP_TO_STRING_CASE(image_gather8h_pck);
 		OP_TO_STRING_CASE(image_sample_d_g16);
 		OP_TO_STRING_CASE(image_sample_d_cl_g16);
 		OP_TO_STRING_CASE(image_sample_c_d_g16);
