RDNA ISA parser improvements and fixes

diff --git a/Engine/Source/Developer/ShaderCompilerCommon/Private/ISAParser/QuadModeDetect.cpp b/Engine/Source/Developer/ShaderCompilerCommon/Private/ISAParser/QuadModeDetect.cpp
--- a/Engine/Source/Developer/ShaderCompilerCommon/Private/ISAParser/QuadModeDetect.cpp
+++ b/Engine/Source/Developer/ShaderCompilerCommon/Private/ISAParser/QuadModeDetect.cpp
@@ -712,6 +712,9 @@ bool ISAParser::HasDerivativeOps(bool& bHasDerivativeOps, const char* Code, uint
 
 	if (InstructionSet == EInstructionSet::RDNA1)
 	{
+		//FDebugParserRDNA1 PrintParser;
+		//PerformParsing(PrintParser, ReadPtr, EndPtr);
+
 		FQuadModeParserRDNA1 Parser;
 		if (PerformParsing(Parser, ReadPtr, EndPtr))
 		{
@@ -724,6 +727,9 @@ bool ISAParser::HasDerivativeOps(bool& bHasDerivativeOps, const char* Code, uint
 	}
 	else if (InstructionSet == EInstructionSet::RDNA2)
 	{
+		//FDebugParserRDNA2 PrintParser;
+		//PerformParsing(PrintParser, ReadPtr, EndPtr);
+
 		FQuadModeParserRDNA2 Parser;
 		if (PerformParsing(Parser, ReadPtr, EndPtr))
 		{
diff --git a/Engine/Source/Developer/ShaderCompilerCommon/Private/ISAParser/RDNA1_ISA.h b/Engine/Source/Developer/ShaderCompilerCommon/Private/ISAParser/RDNA1_ISA.h
--- a/Engine/Source/Developer/ShaderCompilerCommon/Private/ISAParser/RDNA1_ISA.h
+++ b/Engine/Source/Developer/ShaderCompilerCommon/Private/ISAParser/RDNA1_ISA.h
@@ -4,6 +4,7 @@
 
 #include "CoreMinimal.h"
 #include "RDNA_Common.h"
+#include "HAL/OutputDevices.h"
 
 // https://www.amd.com/content/dam/amd/en/documents/radeon-tech-docs/instruction-set-architectures/rdna-shader-instruction-set-architecture.pdf
 
@@ -2315,61 +2316,61 @@ const char* ToString(EGLOBALOps Op)
 void PrintSMEM(const FInstSMEM& Inst)
 {
 	ESMEMOps Op = (ESMEMOps)Inst.OP;
-	printf("%s\n", ToString(Op));
+	UE_LOG(LogTemp, Warning, TEXT("%hs"), ToString(Op));
 }
 
 void PrintSOP1(const FInstSOP1& Inst)
 {
 	ESOP1Ops Op = (ESOP1Ops)Inst.OP;
-	printf("%s\n", ToString(Op));
+	UE_LOG(LogTemp, Warning, TEXT("%hs"), ToString(Op));
 }
 
 void PrintSOP2(const FInstSOP2& Inst)
 {
 	ESOP2Ops Op = (ESOP2Ops)Inst.OP;
-	printf("%s\n", ToString(Op));
+	UE_LOG(LogTemp, Warning, TEXT("%hs"), ToString(Op));
 }
 
 void PrintSOPP(const FInstSOPP& Inst)
 {
 	ESOPPOps Op = (ESOPPOps)Inst.OP;
-	printf("%s\n", ToString(Op));
+	UE_LOG(LogTemp, Warning, TEXT("%hs"), ToString(Op));
 }
 
 void PrintSOPC(const FInstSOPC& Inst)
 {
 	ESOPCOps Op = (ESOPCOps)Inst.OP;
-	printf("%s\n", ToString(Op));
+	UE_LOG(LogTemp, Warning, TEXT("%hs"), ToString(Op));
 }
 
 void PrintSOPK(const FInstSOPK& Inst)
 {
 	ESOPKOps Op = (ESOPKOps)Inst.OP;
-	printf("%s\n", ToString(Op));
+	UE_LOG(LogTemp, Warning, TEXT("%hs"), ToString(Op));
 }
 
 void PrintMIMG(const FInstMIMG& Inst)
 {
 	EMIMGOps Op = (EMIMGOps)Inst.OP;
-	printf("%s\n", ToString(Op));
+	UE_LOG(LogTemp, Warning, TEXT("%hs"), ToString(Op));
 }
 
 void PrintVINTERP(const FInstVINTERP& Inst)
 {
 	EVINTERPOps Op = (EVINTERPOps)Inst.OP;
-	printf("%s\n", ToString(Op));
+	UE_LOG(LogTemp, Warning, TEXT("%hs"), ToString(Op));
 }
 
 void PrintVOP1(const FInstVOP1& Inst)
 {
 	EVOP1Ops Op = (EVOP1Ops)Inst.OP;
-	printf("%s\n", ToString(Op));
+	UE_LOG(LogTemp, Warning, TEXT("%hs"), ToString(Op));
 }
 
 void PrintVOP2(const FInstVOP2& Inst)
 {
 	EVOP2Ops Op = (EVOP2Ops)Inst.OP;
-	printf("%s\n", ToString(Op));
+	UE_LOG(LogTemp, Warning, TEXT("%hs"), ToString(Op));
 }
 
 void PrintVOP3A(const FInstVOP3A& Inst)
@@ -2381,85 +2382,85 @@ void PrintVOP3A(const FInstVOP3A& Inst)
 	{
 		// VOP1 encoded as VOP3A
 		EVOP1Ops Op = (EVOP1Ops)(Inst.OP - VOP1Offset);
-		printf("%s\n", ToString(Op));
+		UE_LOG(LogTemp, Warning, TEXT("%hs"), ToString(Op));
 	}
 	else if (Inst.OP >= VOP2Offset && (Inst.OP - VOP2Offset) <= 60 /* v_pk_fmac_f16 */)
 	{
 		// VOP2 encoded as VOP3A
 		EVOP2Ops Op = (EVOP2Ops)(Inst.OP - VOP2Offset);
-		printf("%s\n", ToString(Op));
+		UE_LOG(LogTemp, Warning, TEXT("%hs"), ToString(Op));
 	}
 	else if (Inst.OP <= 255 /* v_cmpx_tru_f16 */)
 	{
 		// VOPC encoded as VOP3A
 		EVOPCOps Op = (EVOPCOps)Inst.OP;
-		printf("%s\n", ToString(Op));
+		UE_LOG(LogTemp, Warning, TEXT("%hs"), ToString(Op));
 	}
 	else
 	{
 		EVOP3ABOps Op = (EVOP3ABOps)Inst.OP;
-		printf("%s\n", ToString(Op));
+		UE_LOG(LogTemp, Warning, TEXT("%hs"), ToString(Op));
 	}
 }
 
 void PrintVOP3B(const FInstVOP3B& Inst)
 {
 	EVOP3ABOps Op = (EVOP3ABOps)Inst.OP;
-	printf("%s\n", ToString(Op));
+	UE_LOG(LogTemp, Warning, TEXT("%hs"), ToString(Op));
 }
 
 void PrintVOP3P(const FInstVOP3P& Inst)
 {
 	EVOP3POps Op = (EVOP3POps)Inst.OP;
-	printf("%s\n", ToString(Op));
+	UE_LOG(LogTemp, Warning, TEXT("%hs"), ToString(Op));
 }
 
 void PrintVOPC(const FInstVOPC& Inst)
 {
 	EVOPCOps Op = (EVOPCOps)Inst.OP;
-	printf("%s\n", ToString(Op));
+	UE_LOG(LogTemp, Warning, TEXT("%hs"), ToString(Op));
 }
 
 void PrintLDSGDS(const FInstLDSGDS& Inst)
 {
 	//ELDSGDSOps Op = (ELDSGDSOps)Inst.OP;
-	//printf("%s\n", ToString(Op));
-	printf("LDSGDS - todo\n");
+	//UE_LOG(LogTemp, Warning, TEXT("%hs"), ToString(Op));
+	UE_LOG(LogTemp, Warning, TEXT("LDSGDS - todo"));
 }
 
 void PrintMUBUF(const FInstMUBUF& Inst)
 {
 	EMUBUFOps Op = (EMUBUFOps)Inst.OP;
-	printf("%s\n", ToString(Op));
+	UE_LOG(LogTemp, Warning, TEXT("%hs"), ToString(Op));
 }
 
 void PrintMTBUF(const FInstMTBUF& Inst)
 {
 	EMTBUFOps Op = (EMTBUFOps)Inst.OP;
-	printf("%s\n", ToString(Op));
+	UE_LOG(LogTemp, Warning, TEXT("%hs"), ToString(Op));
 }
 
 void PrintEXPORT(const FInstEXPORT& Inst)
 {
-	printf("exp\n");
+	UE_LOG(LogTemp, Warning, TEXT("exp"));
 }
 
 void PrintFLAT(const FInstFSG& Inst)
 {
 	EFLATOps Op = (EFLATOps)Inst.OP;
-	printf("%s\n", ToString(Op));
+	UE_LOG(LogTemp, Warning, TEXT("%hs"), ToString(Op));
 }
 
 void PrintSCRATCH(const FInstFSG& Inst)
 {
 	ESCRATCHOps Op = (ESCRATCHOps)Inst.OP;
-	printf("%s\n", ToString(Op));
+	UE_LOG(LogTemp, Warning, TEXT("%hs"), ToString(Op));
 }
 
 void PrintGLOBAL(const FInstFSG& Inst)
 {
 	EGLOBALOps Op = (EGLOBALOps)Inst.OP;
-	printf("%s\n", ToString(Op));
+	UE_LOG(LogTemp, Warning, TEXT("%hs"), ToString(Op));
 }
 
 void PrintFSG(const FInstFSG& Inst)
diff --git a/Engine/Source/Developer/ShaderCompilerCommon/Private/ISAParser/RDNA2_ISA.h b/Engine/Source/Developer/ShaderCompilerCommon/Private/ISAParser/RDNA2_ISA.h
--- a/Engine/Source/Developer/ShaderCompilerCommon/Private/ISAParser/RDNA2_ISA.h
+++ b/Engine/Source/Developer/ShaderCompilerCommon/Private/ISAParser/RDNA2_ISA.h
@@ -4,6 +4,7 @@
 
 #include "CoreMinimal.h"
 #include "RDNA_Common.h"
+#include "HAL/OutputDevices.h"
 
 // https://www.amd.com/content/dam/amd/en/documents/radeon-tech-docs/instruction-set-architectures/rdna2-shader-instruction-set-architecture.pdf
 
@@ -187,6 +188,7 @@ enum class ESOP2Ops : uint16
 	s_nor_b64 = 27,
 	s_xnor_b32 = 28,
 	s_xnor_b64 = 29,
+	s_lshl_b32 = 30,
 	s_lshl_b64 = 31,
 	s_lshr_b32 = 32,
 	s_lshr_b64 = 33,
@@ -245,6 +247,7 @@ const char* ToString(ESOP2Ops Op)
 		OP_TO_STRING_CASE(s_nor_b64);
 		OP_TO_STRING_CASE(s_xnor_b32);
 		OP_TO_STRING_CASE(s_xnor_b64);
+		OP_TO_STRING_CASE(s_lshl_b32);
 		OP_TO_STRING_CASE(s_lshl_b64);
 		OP_TO_STRING_CASE(s_lshr_b32);
 		OP_TO_STRING_CASE(s_lshr_b64);
@@ -2315,61 +2318,61 @@ const char* ToString(EGLOBALOps Op)
 void PrintSMEM(const FInstSMEM& Inst)
 {
 	ESMEMOps Op = (ESMEMOps)Inst.OP;
-	printf("%s\n", ToString(Op));
+	UE_LOG(LogTemp, Warning, TEXT("%hs"), ToString(Op));
 }
 
 void PrintSOP1(const FInstSOP1& Inst)
 {
 	ESOP1Ops Op = (ESOP1Ops)Inst.OP;
-	printf("%s\n", ToString(Op));
+	UE_LOG(LogTemp, Warning, TEXT("%hs"), ToString(Op));
 }
 
 void PrintSOP2(const FInstSOP2& Inst)
 {
 	ESOP2Ops Op = (ESOP2Ops)Inst.OP;
-	printf("%s\n", ToString(Op));
+	UE_LOG(LogTemp, Warning, TEXT("%hs"), ToString(Op));
 }
 
 void PrintSOPP(const FInstSOPP& Inst)
 {
 	ESOPPOps Op = (ESOPPOps)Inst.OP;
-	printf("%s\n", ToString(Op));
+	UE_LOG(LogTemp, Warning, TEXT("%hs"), ToString(Op));
 }
 
 void PrintSOPC(const FInstSOPC& Inst)
 {
 	ESOPCOps Op = (ESOPCOps)Inst.OP;
-	printf("%s\n", ToString(Op));
+	UE_LOG(LogTemp, Warning, TEXT("%hs"), ToString(Op));
 }
 
 void PrintSOPK(const FInstSOPK& Inst)
 {
 	ESOPKOps Op = (ESOPKOps)Inst.OP;
-	printf("%s\n", ToString(Op));
+	UE_LOG(LogTemp, Warning, TEXT("%hs"), ToString(Op));
 }
 
 void PrintMIMG(const FInstMIMG& Inst)
 {
 	EMIMGOps Op = (EMIMGOps)Inst.OP;
-	printf("%s\n", ToString(Op));
+	UE_LOG(LogTemp, Warning, TEXT("%hs"), ToString(Op));
 }
 
 void PrintVINTERP(const FInstVINTERP& Inst)
 {
 	EVINTERPOps Op = (EVINTERPOps)Inst.OP;
-	printf("%s\n", ToString(Op));
+	UE_LOG(LogTemp, Warning, TEXT("%hs"), ToString(Op));
 }
 
 void PrintVOP1(const FInstVOP1& Inst)
 {
 	EVOP1Ops Op = (EVOP1Ops)Inst.OP;
-	printf("%s\n", ToString(Op));
+	UE_LOG(LogTemp, Warning, TEXT("%hs"), ToString(Op));
 }
 
 void PrintVOP2(const FInstVOP2& Inst)
 {
 	EVOP2Ops Op = (EVOP2Ops)Inst.OP;
-	printf("%s\n", ToString(Op));
+	UE_LOG(LogTemp, Warning, TEXT("%hs"), ToString(Op));
 }
 
 void PrintVOP3A(const FInstVOP3A& Inst)
@@ -2381,85 +2384,85 @@ void PrintVOP3A(const FInstVOP3A& Inst)
 	{
 		// VOP1 encoded as VOP3A
 		EVOP1Ops Op = (EVOP1Ops)(Inst.OP - VOP1Offset);
-		printf("%s\n", ToString(Op));
+		UE_LOG(LogTemp, Warning, TEXT("%hs"), ToString(Op));
 	}
 	else if (Inst.OP >= VOP2Offset && (Inst.OP - VOP2Offset) <= 60 /* v_pk_fmac_f16 */)
 	{
 		// VOP2 encoded as VOP3A
 		EVOP2Ops Op = (EVOP2Ops)(Inst.OP - VOP2Offset);
-		printf("%s\n", ToString(Op));
+		UE_LOG(LogTemp, Warning, TEXT("%hs"), ToString(Op));
 	}
 	else if (Inst.OP <= 255 /* v_cmpx_tru_f16 */)
 	{
 		// VOPC encoded as VOP3A
 		EVOPCOps Op = (EVOPCOps)Inst.OP;
-		printf("%s\n", ToString(Op));
+		UE_LOG(LogTemp, Warning, TEXT("%hs"), ToString(Op));
 	}
 	else
 	{
 		EVOP3ABOps Op = (EVOP3ABOps)Inst.OP;
-		printf("%s\n", ToString(Op));
+		UE_LOG(LogTemp, Warning, TEXT("%hs"), ToString(Op));
 	}
 }
 
 void PrintVOP3B(const FInstVOP3B& Inst)
 {
 	EVOP3ABOps Op = (EVOP3ABOps)Inst.OP;
-	printf("%s\n", ToString(Op));
+	UE_LOG(LogTemp, Warning, TEXT("%hs"), ToString(Op));
 }
 
 void PrintVOP3P(const FInstVOP3P& Inst)
 {
 	EVOP3POps Op = (EVOP3POps)Inst.OP;
-	printf("%s\n", ToString(Op));
+	UE_LOG(LogTemp, Warning, TEXT("%hs"), ToString(Op));
 }
 
 void PrintVOPC(const FInstVOPC& Inst)
 {
 	EVOPCOps Op = (EVOPCOps)Inst.OP;
-	printf("%s\n", ToString(Op));
+	UE_LOG(LogTemp, Warning, TEXT("%hs"), ToString(Op));
 }
 
 void PrintLDSGDS(const FInstLDSGDS& Inst)
 {
 	//ELDSGDSOps Op = (ELDSGDSOps)Inst.OP;
-	//printf("%s\n", ToString(Op));
-	printf("LDSGDS - todo\n");
+	//UE_LOG(LogTemp, Warning, TEXT("%hs"), ToString(Op));
+	UE_LOG(LogTemp, Warning, TEXT("LDSGDS - todo"));
 }
 
 void PrintMUBUF(const FInstMUBUF& Inst)
 {
 	EMUBUFOps Op = (EMUBUFOps)Inst.OP;
-	printf("%s\n", ToString(Op));
+	UE_LOG(LogTemp, Warning, TEXT("%hs"), ToString(Op));
 }
 
 void PrintMTBUF(const FInstMTBUF& Inst)
 {
 	EMTBUFOps Op = (EMTBUFOps)Inst.OP;
-	printf("%s\n", ToString(Op));
+	UE_LOG(LogTemp, Warning, TEXT("%hs"), ToString(Op));
 }
 
 void PrintEXPORT(const FInstEXPORT& Inst)
 {
-	printf("exp\n");
+	UE_LOG(LogTemp, Warning, TEXT("exp"));
 }
 
 void PrintFLAT(const FInstFSG& Inst)
 {
 	EFLATOps Op = (EFLATOps)Inst.OP;
-	printf("%s\n", ToString(Op));
+	UE_LOG(LogTemp, Warning, TEXT("%hs"), ToString(Op));
 }
 
 void PrintSCRATCH(const FInstFSG& Inst)
 {
 	ESCRATCHOps Op = (ESCRATCHOps)Inst.OP;
-	printf("%s\n", ToString(Op));
+	UE_LOG(LogTemp, Warning, TEXT("%hs"), ToString(Op));
 }
 
 void PrintGLOBAL(const FInstFSG& Inst)
 {
 	EGLOBALOps Op = (EGLOBALOps)Inst.OP;
-	printf("%s\n", ToString(Op));
+	UE_LOG(LogTemp, Warning, TEXT("%hs"), ToString(Op));
 }
 
 void PrintFSG(const FInstFSG& Inst)
diff --git a/Engine/Source/Developer/ShaderCompilerCommon/Private/ISAParser/RDNA_Common.h b/Engine/Source/Developer/ShaderCompilerCommon/Private/ISAParser/RDNA_Common.h
--- a/Engine/Source/Developer/ShaderCompilerCommon/Private/ISAParser/RDNA_Common.h
+++ b/Engine/Source/Developer/ShaderCompilerCommon/Private/ISAParser/RDNA_Common.h
@@ -478,7 +478,7 @@ bool IsVOP3(const uint32* Code, uint32& AdvanceAmount)
 	if (bMatch)
 	{
 		AdvanceAmount = GetAdvanceAmount(EInstructionType::VOP3);
-		if (VOP3.SRC0 == Operand_DPP8 || VOP3.SRC0 == Operand_DPP8FI || VOP3.SRC0 == Operand_DPP16 || VOP3.SRC0 == Operand_SDWA || VOP3.SRC0 == Operand_Literal)
+		if (VOP3.SRC0 == Operand_DPP8 || VOP3.SRC0 == Operand_DPP8FI || VOP3.SRC0 == Operand_DPP16 || VOP3.SRC0 == Operand_SDWA || VOP3.SRC0 == Operand_Literal || VOP3.SRC2 == Operand_Literal)
 		{
 			AdvanceAmount++;
 		}
