Removing redundant line directives from Shader Minifier.

diff --git a/Engine/Source/Developer/ShaderCompilerCommon/Private/ShaderMinifier.cpp b/Engine/Source/Developer/ShaderCompilerCommon/Private/ShaderMinifier.cpp
--- a/Engine/Source/Developer/ShaderCompilerCommon/Private/ShaderMinifier.cpp
+++ b/Engine/Source/Developer/ShaderCompilerCommon/Private/ShaderMinifier.cpp
@@ -1497,6 +1497,9 @@ static FString MinifyShader(const FParsedShader& Parsed, TConstArrayView<FString
 
 	const FNamespace* CurrentNamespace = nullptr;
 
+	int32 LastLineNumber = -1;
+	FStringView LastLineFileName;
+
 	for (const FCodeChunk& Chunk : Parsed.Chunks)
 	{
 		auto ShouldSkipChunk = [&RelevantChunks, &Chunk, Flags]()
@@ -1574,8 +1577,13 @@ static FString MinifyShader(const FParsedShader& Parsed, TConstArrayView<FString
 			if (ChunkLine != INDEX_NONE && LineDirectiveIndex == INDEX_NONE)
 			{
 				// There was no valid line directive for this chunk, but we do know the line in the input source, so just emit that.
+				if (ChunkLine > LastLineNumber + 1 || !LastLineFileName.IsEmpty())
+				{
+					OutputStream << TEXT("#line ") << ChunkLine << TEXT("\n");
+				}
 
-				OutputStream << TEXT("#line ") << ChunkLine << TEXT("\n");
+				LastLineNumber = ChunkLine;
+				LastLineFileName = FStringView();
 			}
 			else if (ChunkLine != INDEX_NONE && LineDirectiveIndex != INDEX_NONE)
 			{
@@ -1592,19 +1600,28 @@ static FString MinifyShader(const FParsedShader& Parsed, TConstArrayView<FString
 					int32 OffsetFromLineDirective = ChunkLine - (LineDirectiveLine + 1); // Line directive identifies the *next* line, hence +1 when computing the offset
 					int32 PatchedLineNumber = ParsedLineNumber + OffsetFromLineDirective;
 
-					OutputStream << TEXT("#line ") << PatchedLineNumber;
-					if (!ParsedFileName.IsEmpty())
+					if (PatchedLineNumber > LastLineNumber + 1 || LastLineFileName != ParsedFileName)
 					{
-						OutputStream << TEXT(" \"") << ParsedFileName << TEXT("\"");
+						// Separate the next block from the previous one when it starts with a line directive
+						if (OutputStream.Len())
+						{
+							OutputStream << "\n";
+						}
+						OutputStream << TEXT("#line ") << PatchedLineNumber;
+						if (!ParsedFileName.IsEmpty())
+						{
+							OutputStream << TEXT(" \"") << ParsedFileName << TEXT("\"");
+						}
+						OutputStream << TEXT("\n");
 					}
-					OutputStream << TEXT("\n");
+
+					LastLineNumber = PatchedLineNumber;
+					LastLineFileName = ParsedFileName;
 				}
 			}
 		}
 
 		OutputChunk(Chunk, OutputStream);
-
-		OutputStream << "\n";
 	}
 
 	if (CurrentNamespace)
