Allow passing custom SPIV-V optimization passes to ShaderConductor. This allows us to control more fine-grained what optimization passes work best for certain render passes on different platforms.

#rnx

diff --git a/Engine/Source/Developer/ShaderCompilerCommon/Private/ShaderConductorContext.cpp b/Engine/Source/Developer/ShaderCompilerCommon/Private/ShaderConductorContext.cpp
--- a/Engine/Source/Developer/ShaderCompilerCommon/Private/ShaderConductorContext.cpp
+++ b/Engine/Source/Developer/ShaderCompilerCommon/Private/ShaderConductorContext.cpp
@@ -132,6 +132,7 @@ namespace CrossCompiler
 		TArray<ShaderConductor::MacroDefine> DefineRefs;
 		TArray<TPair<TArray<ANSICHAR>, TArray<ANSICHAR>>> Flags;
 		TArray<ShaderConductor::MacroDefine> FlagRefs;
+		TArray<ANSICHAR> InternalDxcArgs;
 		TArray<TArray<ANSICHAR>> CustomDxcArgs;
 		TArray<ANSICHAR const*> CustomDxcArgRefs;
 		TArray<ANSICHAR const*> DxcArgRefs;
@@ -319,7 +320,7 @@ namespace CrossCompiler
 		}
 	}
 
-	static void ConvertScOptions(const FShaderConductorOptions& InOptions, ShaderConductor::Compiler::Options& OutOptions, const TArray<ANSICHAR const*>& CustomDxcArgRefs, TArray<ANSICHAR const*>& DxcArgRefs)
+	static void ConvertScOptions(FShaderConductorContext::FShaderConductorIntermediates& Intermediates, const FShaderConductorOptions& InOptions, ShaderConductor::Compiler::Options& OutOptions, bool bIgnoreCustomDxcArgs = false)
 	{
 		// Validate input shader model with respect to certain language features.
 		checkf(
@@ -342,6 +343,8 @@ namespace CrossCompiler
 			static_cast<uint8>(InOptions.ShaderModel.Minor)
 		};
 
+		TArray<ANSICHAR const*>& DxcArgRefs = Intermediates.DxcArgRefs;
+
 		DxcArgRefs.Empty();
 
 		// Select language version
@@ -396,6 +399,9 @@ namespace CrossCompiler
 		if (OutOptions.enable16bitTypes)
 		{
 			DxcArgRefs.Add("-fspv-target-env=universal1.5");
+
+			// ShaderConductor.cpp forgot to pipedown enable16bitTypes, so work arround by adding the parameter manually in here.
+			DxcArgRefs.Add("-enable-16bit-types");
 		}
 		else
 		{
@@ -418,18 +424,27 @@ namespace CrossCompiler
 			}
 		}
 
+		if (!InOptions.SpirvCustomOptimizationPasses.IsEmpty())
+		{
+			ConvertFStringToAnsiString(FString::Printf(TEXT("-Oconfig=%s"), *InOptions.SpirvCustomOptimizationPasses), Intermediates.InternalDxcArgs);
+			DxcArgRefs.Add(Intermediates.InternalDxcArgs.GetData());
+		}
+
 		if (DxcArgRefs.Num() > 0)
 		{
 			// Use DXC argument container and append custom arguments
-			DxcArgRefs.Append(CustomDxcArgRefs);
+			if (!bIgnoreCustomDxcArgs)
+			{
+				DxcArgRefs.Append(Intermediates.CustomDxcArgRefs);
+			}
 			OutOptions.numDXCArgs = DxcArgRefs.Num();
 			OutOptions.DXCArgs = (const char**)DxcArgRefs.GetData();
 		}
-		else if (CustomDxcArgRefs.Num() > 0)
+		else if (!bIgnoreCustomDxcArgs)
 		{
 			// Use custom DXC arguments only
-			OutOptions.numDXCArgs = CustomDxcArgRefs.Num();
-			OutOptions.DXCArgs = (const char**)CustomDxcArgRefs.GetData();
+			OutOptions.numDXCArgs = Intermediates.CustomDxcArgRefs.Num();
+			OutOptions.DXCArgs = (const char**)Intermediates.CustomDxcArgRefs.GetData();
 		}
 		else
 		{
@@ -567,18 +582,8 @@ namespace CrossCompiler
 		ConvertScSourceDesc(*Intermediates, ScSourceDesc);
 
 		ShaderConductor::Compiler::Options ScOptions;
-		ConvertScOptions(Options, ScOptions, {}, Intermediates->DxcArgRefs);
-
-		// ShaderConductor.cpp forgot to pipedown enable16bitTypes, so work arround by adding the parameter manually in here.
-		TArray<ANSICHAR const*> PatchedDxcArgRefs;
-		if (ScOptions.enable16bitTypes)
-		{
-			PatchedDxcArgRefs = Intermediates->DxcArgRefs;
-			PatchedDxcArgRefs.Add("-enable-16bit-types");
-
-			ScOptions.numDXCArgs = PatchedDxcArgRefs.Num();
-			ScOptions.DXCArgs = (const char**)PatchedDxcArgRefs.GetData();
-		}
+		constexpr bool bIgnoreExtraDxcArgs = true;
+		ConvertScOptions(*Intermediates, Options, ScOptions, bIgnoreExtraDxcArgs);
 
 		// Rewrite HLSL with wrapper function to catch exceptions from ShaderConductor
 		bool bSucceeded = false;
@@ -615,7 +620,7 @@ namespace CrossCompiler
 		ScTargetDesc.language = ShaderConductor::ShadingLanguage::SpirV;
 
 		ShaderConductor::Compiler::Options ScOptions;
-		ConvertScOptions(Options, ScOptions, Intermediates->CustomDxcArgRefs, Intermediates->DxcArgRefs);
+		ConvertScOptions(*Intermediates, Options, ScOptions);
 
 		// Compile HLSL source code to SPIR-V
 		bool bSucceeded = false;
@@ -711,7 +716,7 @@ namespace CrossCompiler
 		ConvertScTargetDesc(*Intermediates, Target, ScTargetDesc);
 
 		ShaderConductor::Compiler::Options ScOptions;
-		ConvertScOptions(Options, ScOptions, Intermediates->CustomDxcArgRefs, Intermediates->DxcArgRefs);
+		ConvertScOptions(*Intermediates, Options, ScOptions);
 
 		ShaderConductor::Compiler::ResultDesc ScBinaryDesc;
 		ScBinaryDesc.target.Reset(InSpirv, InSpirvByteSize);
diff --git a/Engine/Source/Developer/ShaderCompilerCommon/Public/ShaderConductorContext.h b/Engine/Source/Developer/ShaderCompilerCommon/Public/ShaderConductorContext.h
--- a/Engine/Source/Developer/ShaderCompilerCommon/Public/ShaderConductorContext.h
+++ b/Engine/Source/Developer/ShaderCompilerCommon/Public/ShaderConductorContext.h
@@ -110,6 +110,9 @@ namespace CrossCompiler
 
 		/** HLSL language input version: 2015, 2016, 2017, 2018 (Default), 2021 (Breaking changes in short-circuiting evaluation). */
 		uint32 HlslVersion = 2018;
+
+		/** SPIR-V specific optimization passes to override the default '-O' argument. This will be passed to DXC via the '-Oconfig=...' argument. */
+		FString SpirvCustomOptimizationPasses;
 	};
 
 	/** Target high level languages for ShaderConductor output. */
