Platform-specific DXC improvements:
- Avoid explicit FMA instructions in platform-specific GLSL to let backend-compiler decide when to use them.
- Fix code refelction when using SPIR-V directly with backend-compiler.
- Define preset of SPIR-V optimization passes to reduce nested expressions "preset(relax-nested-expr)".

diff --git a/Engine/Source/Developer/ShaderCompilerCommon/Private/ShaderConductorContext.cpp b/Engine/Source/Developer/ShaderCompilerCommon/Private/ShaderConductorContext.cpp
--- a/Engine/Source/Developer/ShaderCompilerCommon/Private/ShaderConductorContext.cpp
+++ b/Engine/Source/Developer/ShaderCompilerCommon/Private/ShaderConductorContext.cpp
@@ -320,6 +320,71 @@ namespace CrossCompiler
 		}
 	}
 
+
+	/*
+	 * Reduced list of SPIRV-Tools optimization passes to avoid nested expressions in GLSL output.
+	 * The order of passes has been adopted from the standard '-O' optimization configuration with the following passes removed:
+	 * - LocalSingleBlockLoadStoreElimPass
+	 * - LocalSingleStoreElimPass
+	 * - LocalMultiStoreElimPass
+	 * - SSARewritePass
+	 */
+	static const TCHAR* GSpirvTools_OptPasses_PresetRelaxNestedExpr = TEXT(
+		"--wrap-opkill,"
+		"--eliminate-dead-branches,"
+		"--merge-return,"
+		"--inline-entry-points-exhaustive,"
+		"--eliminate-dead-functions,"
+		"--eliminate-dead-code-aggressive,"
+		"--private-to-local,"
+		"--eliminate-dead-code-aggressive,"
+		"--scalar-replacement,"
+		"--convert-local-access-chains,"
+		"--eliminate-dead-code-aggressive,"
+		"--ccp,"
+		"--eliminate-dead-code-aggressive,"
+		"--loop-unroll,"
+		"--eliminate-dead-branches,"
+		"--redundancy-elimination,"
+		"--combine-access-chains,"
+		"--simplify-instructions,"
+		"--scalar-replacement,"
+		"--convert-local-access-chains,"
+		"--eliminate-dead-code-aggressive,"
+		"--vector-dce,"
+		"--eliminate-dead-inserts,"
+		"--eliminate-dead-branches,"
+		"--simplify-instructions,"
+		"--if-conversion,"
+		"--copy-propagate-arrays,"
+		"--reduce-load-size,"
+		"--eliminate-dead-code-aggressive,"
+		"--merge-blocks,"
+		"--redundancy-elimination,"
+		"--eliminate-dead-branches,"
+		"--merge-blocks,"
+		"--simplify-instructions,"
+		"--eliminate-dead-members,"
+		"--merge-blocks,"
+		"--redundancy-elimination,"
+		"--simplify-instructions,"
+		"--eliminate-dead-code-aggressive,"
+		"--cfg-cleanup"
+	);
+
+	static const TCHAR* SelectSpirvCustomOptimizationPasses(const FString& OptimizationPasses)
+	{
+		if (OptimizationPasses == TEXT("preset(relax-nested-expr)"))
+		{
+			return GSpirvTools_OptPasses_PresetRelaxNestedExpr;
+		}
+		else
+		{
+			// Interpret input argument as set of optimization passes
+			return *OptimizationPasses;
+		}
+	}
+
 	static void ConvertScOptions(FShaderConductorContext::FShaderConductorIntermediates& Intermediates, const FShaderConductorOptions& InOptions, ShaderConductor::Compiler::Options& OutOptions, bool bIgnoreCustomDxcArgs = false)
 	{
 		// Validate input shader model with respect to certain language features.
@@ -426,7 +491,7 @@ namespace CrossCompiler
 
 		if (!InOptions.SpirvCustomOptimizationPasses.IsEmpty())
 		{
-			ConvertFStringToAnsiString(FString::Printf(TEXT("-Oconfig=%s"), *InOptions.SpirvCustomOptimizationPasses), Intermediates.InternalDxcArgs);
+			ConvertFStringToAnsiString(FString::Printf(TEXT("-Oconfig=%s"), SelectSpirvCustomOptimizationPasses(InOptions.SpirvCustomOptimizationPasses)), Intermediates.InternalDxcArgs);
 			DxcArgRefs.Add(Intermediates.InternalDxcArgs.GetData());
 		}
 
diff --git a/Engine/Source/Developer/ShaderCompilerCommon/Public/ShaderConductorContext.h b/Engine/Source/Developer/ShaderCompilerCommon/Public/ShaderConductorContext.h
--- a/Engine/Source/Developer/ShaderCompilerCommon/Public/ShaderConductorContext.h
+++ b/Engine/Source/Developer/ShaderCompilerCommon/Public/ShaderConductorContext.h
@@ -111,7 +111,10 @@ namespace CrossCompiler
 		/** HLSL language input version: 2015, 2016, 2017, 2018 (Default), 2021 (Breaking changes in short-circuiting evaluation). */
 		uint32 HlslVersion = 2018;
 
-		/** SPIR-V specific optimization passes to override the default '-O' argument. This will be passed to DXC via the '-Oconfig=...' argument. */
+		/**
+		 * SPIR-V specific optimization passes to override the default '-O' argument. This will be passed to DXC via the '-Oconfig=...' argument.
+		 * Use "preset(relax-nested-expr)" for pre-defined set of optimization passes to relax nested expressions.
+		 */
 		FString SpirvCustomOptimizationPasses;
 	};
 
