Improve ShaderSymbols.info file generation in multicook scenarios
Dedupe info where multiple workers built same shader

#tests local runs, compare of farm built files attached in bug

diff --git a/Engine/Source/Developer/ShaderCompilerCommon/Private/ShaderSymbolExport.cpp b/Engine/Source/Developer/ShaderCompilerCommon/Private/ShaderSymbolExport.cpp
--- a/Engine/Source/Developer/ShaderCompilerCommon/Private/ShaderSymbolExport.cpp
+++ b/Engine/Source/Developer/ShaderCompilerCommon/Private/ShaderSymbolExport.cpp
@@ -208,26 +208,13 @@ void FShaderSymbolExport::NotifyShaderCompilersShutdown()
 {
 	if (ShaderInfos.Num())
 	{
-		// sort and combine the data, to either write to a zip or just as a file
-		ShaderInfos.StableSort([](const FShaderInfo& A, const FShaderInfo& B) { return A.Hash < B.Hash; });
-
-		TArray<uint8> Output;
-		for (FShaderInfo Info : ShaderInfos)
-		{
-			auto TmpHash = StringCast<ANSICHAR>(*Info.Hash);
-			auto TmpData = StringCast<ANSICHAR>(*Info.Data);
-			Output.Append((const uint8*)TmpHash.Get(), TmpHash.Length());
-			Output.Add(' ');
-			Output.Append((const uint8*)TmpData.Get(), TmpData.Length());
-			Output.Add('\n');
-		}
-
 		if (InfoFilePath.Len())
 		{
 			IFileManager& FileManager = IFileManager::Get();
 			if (bMultiprocessOwner)
 			{
 				// if we are the multiprocess owner merge in any other files we find
+				// we will chunk up the worker files into {Hash, Data} pairs, dedupe them with ours, and sort them all
 				IPlatformFile& PlatformFile = FPlatformFileManager::Get().GetPlatformFile();
 				TArray<FString> FilesToMergeIn;
 				PlatformFile.FindFiles(FilesToMergeIn, *ExportPath, InfoFileExtension);
@@ -241,12 +228,47 @@ void FShaderSymbolExport::NotifyShaderCompilersShutdown()
 						RawData.AddUninitialized(Size);
 						Reader->Serialize(RawData.GetData(), Size);
 						Reader->Close();
-						Output.Append(RawData);
+
+						TArray<FString> Lines;
+						FString(StringCast<TCHAR>(reinterpret_cast<const ANSICHAR*>(RawData.GetData())).Get()).ParseIntoArrayLines(Lines);
+
+						for (const FString& Line : Lines)
+						{
+							int32 Space;
+							Line.FindChar(TEXT(' '), Space);
+							if (Space != INDEX_NONE)
+							{
+								FString Filename = Line.Left(Space);
+
+								// if this symbol is new to the multiproc owner, store it
+								bool bAlreadyInSet = false;
+								ExportedShaders.Add(Filename, &bAlreadyInSet);
+								if (!bAlreadyInSet)
+								{
+									FString DebugData = Line.Right(Line.Len() - Space - 1);
+									ShaderInfos.Add({ Filename, DebugData });
+								}
+							}
+						}
 					}
 					PlatformFile.DeleteFile(*InfoFile);
 				}
 			}
 
+			// sort and combine the data for output
+			ShaderInfos.Sort([](const FShaderInfo& A, const FShaderInfo& B) { return A.Hash < B.Hash; });
+
+			TArray<uint8> Output;
+			for (FShaderInfo Info : ShaderInfos)
+			{
+				auto TmpHash = StringCast<ANSICHAR>(*Info.Hash);
+				auto TmpData = StringCast<ANSICHAR>(*Info.Data);
+				Output.Append((const uint8*)TmpHash.Get(), TmpHash.Length());
+				Output.Add(' ');
+				Output.Append((const uint8*)TmpData.Get(), TmpData.Length());
+				Output.Add('\n');
+			}
+
 			TUniquePtr<FArchive> Writer = TUniquePtr<FArchive>(FileManager.CreateFileWriter(*InfoFilePath));
 			if (Writer.IsValid())
 			{
