Shader preprocessing optimizations:
- Removing system calls from UniformBuffer parsing and processing.
- When searching for uses of a UniformBuffer, start after the end of the UniformBuffer declaration.
- Simplifying the code around UniformBuffer reference cleanup
- Cuts time in CleanupUniformBufferCode by 50-75% based on how large the shader is.

diff --git a/Engine/Source/Developer/ShaderCompilerCommon/Private/ShaderCompilerCommon.cpp b/Engine/Source/Developer/ShaderCompilerCommon/Private/ShaderCompilerCommon.cpp
--- a/Engine/Source/Developer/ShaderCompilerCommon/Private/ShaderCompilerCommon.cpp
+++ b/Engine/Source/Developer/ShaderCompilerCommon/Private/ShaderCompilerCommon.cpp
@@ -213,14 +213,29 @@ void CullGlobalUniformBuffers(const TMap<FString, FUniformBufferEntry>& UniformB
 	}
 }
 
+static bool IsSpaceOrTabOrEOL(TCHAR Char)
+{
+	return Char == ' ' || Char == '\t' || Char == '\n' || Char == '\r';
+}
+
+static const TCHAR* FindNextChar(const TCHAR* ReadStart, TCHAR SearchChar)
+{
+	const TCHAR* SearchPtr = ReadStart;
+	while (*SearchPtr && *SearchPtr != SearchChar)
+	{
+		SearchPtr++;
+	}
+	return SearchPtr;
+}
+
 const TCHAR* FindNextWhitespace(const TCHAR* StringPtr)
 {
-	while (*StringPtr && !FChar::IsWhitespace(*StringPtr))
+	while (*StringPtr && !IsSpaceOrTabOrEOL(*StringPtr))
 	{
 		StringPtr++;
 	}
 
-	if (*StringPtr && FChar::IsWhitespace(*StringPtr))
+	if (*StringPtr && IsSpaceOrTabOrEOL(*StringPtr))
 	{
 		return StringPtr;
 	}
@@ -232,19 +247,32 @@ const TCHAR* FindNextWhitespace(const TCHAR* StringPtr)
 
 const TCHAR* FindNextNonWhitespace(const TCHAR* StringPtr)
 {
-	while (*StringPtr && FChar::IsWhitespace(*StringPtr))
+	while (*StringPtr && IsSpaceOrTabOrEOL(*StringPtr))
 	{
 		StringPtr++;
 	}
 
-	if (*StringPtr && !FChar::IsWhitespace(*StringPtr))
+	if (*StringPtr && !IsSpaceOrTabOrEOL(*StringPtr))
 	{
 		return StringPtr;
 	}
-	else
+
+	return nullptr;
+}
+
+const TCHAR* FindPreviousNonWhitespace(const TCHAR* StringPtr)
+{
+	do
 	{
-		return nullptr;
+		StringPtr--;
+	} while (*StringPtr && IsSpaceOrTabOrEOL(*StringPtr));
+
+	if (*StringPtr && !IsSpaceOrTabOrEOL(*StringPtr))
+	{
+		return StringPtr;
 	}
+
+	return nullptr;
 }
 
 const TCHAR* FindMatchingBlock(const TCHAR* OpeningCharPtr, char OpenChar, char CloseChar)
@@ -307,7 +335,7 @@ void ParseHLSLTypeName(const TCHAR* SearchString, const TCHAR*& TypeNameStartPtr
 			Depth--;
 		}
 		else if (Depth == 0 
-			&& FChar::IsWhitespace(*TypeNameEndPtr)
+			&& IsSpaceOrTabOrEOL(*TypeNameEndPtr)
 			// If we found a '<', we must not accept any whitespace before it
 			&& (!PotentialExtraTypeInfoPtr || *PotentialExtraTypeInfoPtr != '<' || TypeNameEndPtr > PotentialExtraTypeInfoPtr))
 		{
@@ -343,19 +371,13 @@ const TCHAR* ParseHLSLSymbolName(const TCHAR* SearchString, FString& SymbolName)
 	return Result.GetData() + Result.Len();
 }
 
-bool MatchStructMemberName(const FString& SymbolName, const TCHAR* SearchPtr, const FString& PreprocessedShaderSource)
+const TCHAR* MatchStructMemberName(const FString& SymbolName, const TCHAR* SearchPtr)
 {
-	// Only match whole symbol
-	if (IsValidHLSLIdentifierCharacter(*(SearchPtr - 1)) || *(SearchPtr - 1) == '.')
-	{
-		return false;
-	}
-
 	for (int32 i = 0; i < SymbolName.Len(); i++)
 	{
 		if (*SearchPtr != SymbolName[i])
 		{
-			return false;
+			return nullptr;
 		}
 		
 		SearchPtr++;
@@ -364,7 +386,7 @@ bool MatchStructMemberName(const FString& SymbolName, const TCHAR* SearchPtr, co
 		{
 			// Skip whitespace within the struct member reference before the end
 			// eg 'View. ViewToClip'
-			while (FChar::IsWhitespace(*SearchPtr))
+			while (IsSpaceOrTabOrEOL(*SearchPtr))
 			{
 				SearchPtr++;
 			}
@@ -374,32 +396,10 @@ bool MatchStructMemberName(const FString& SymbolName, const TCHAR* SearchPtr, co
 	// Only match whole symbol
 	if (IsValidHLSLIdentifierCharacter(*SearchPtr))
 	{
-		return false;
+		return nullptr;
 	}
 
-	return true;
-}
-
-// Searches string SearchPtr for 'SearchString.' or 'SearchString .' and returns a pointer to the first character of the match.
-TCHAR* FindNextUniformBufferReference(TCHAR* SearchPtr, const TCHAR* SearchString, uint32 SearchStringLength)
-{
-	TCHAR* FoundPtr = FCString::Strstr(SearchPtr, SearchString);
-	
-	while(FoundPtr)
-	{
-		if (FoundPtr == nullptr)
-		{
-			return nullptr;
-		}
-		else if (FoundPtr[SearchStringLength] == '.' || (FoundPtr[SearchStringLength] == ' ' && FoundPtr[SearchStringLength+1] == '.'))
-		{
-			return FoundPtr;
-		}
-		
-		FoundPtr = FCString::Strstr(FoundPtr + SearchStringLength, SearchString);
-	}
-	
-	return nullptr;
+	return SearchPtr;
 }
 
 EShaderParameterType UE::ShaderCompilerCommon::ParseParameterType(
@@ -800,50 +800,81 @@ struct FUniformBufferMemberInfo
 	FString GlobalName;
 };
 
-const TCHAR* ParseUniformBufferDefinition(const TCHAR* ReadStart, TMap<FString, TArray<FUniformBufferMemberInfo>>& UniformBufferNameToMembers)
+struct FUniformBufferInfo
+{
+	int32 DefinitionEndOffset;
+	TArray<FUniformBufferMemberInfo> Members;
+};
+
+static void FillUniformBufferDefinition(TStringBuilder<128>& Builder, const TCHAR* Start, const TCHAR* End)
 {
-	const FString UniformBufferName(ParseHLSLSymbolName(ReadStart));
+	Builder.Reset();
 
-	const TCHAR* OpeningBrace = FCString::Strstr(ReadStart, TEXT("{"));
+	const TCHAR* Cursor = Start;
+	
+	while (Cursor != End)
+	{
+		if (!IsSpaceOrTabOrEOL(*Cursor))
+		{
+			Builder.AppendChar(*Cursor);
+		}
+
+		Cursor++;
+	}
+}
+
+const TCHAR* ParseUniformBufferDefinition(const TCHAR* SourceStart, const TCHAR* ReadStart, TMap<FStringView, FUniformBufferInfo>& UniformBufferInfos)
+{
+	const FStringView UniformBufferName(ParseHLSLSymbolName(ReadStart));
+
+	const TCHAR* OpeningBrace = FindNextChar(ReadStart, '{');
 	const TCHAR* ClosingBrace = FindMatchingClosingBrace(OpeningBrace + 1);
 
 	const TCHAR* CurrentParseStart = OpeningBrace + 1;
-	const TCHAR* NextSemicolon = FCString::Strstr(CurrentParseStart, TEXT(";"));
+	const TCHAR* NextSemicolon = FindNextChar(CurrentParseStart, ';');
+
+	TStringBuilder<128> StructName;
+	TStringBuilder<128> GlobalName;
 	
 	while (NextSemicolon < ClosingBrace)
 	{
-		const TCHAR* NextEquals = FCString::Strstr(CurrentParseStart, TEXT("="));
-		if (NextEquals < NextSemicolon)
+		const TCHAR* NextSeparator = FindNextChar(CurrentParseStart, '=');
+		if (NextSeparator < NextSemicolon)
 		{
 			const TCHAR* StructStart = CurrentParseStart;
-			const TCHAR* StructEnd = NextEquals - 1;
+			const TCHAR* StructEnd = NextSeparator - 1;
 
-			const TCHAR* GlobalStart = NextEquals + 1;
+			const TCHAR* GlobalStart = NextSeparator + 1;
 			const TCHAR* GlobalEnd = NextSemicolon - 1;
 
-			FString StructName(StructEnd - StructStart, StructStart);
-			StructName.TrimStartAndEndInline();
-			StructName.ReplaceInline(TEXT(" "), TEXT(""));
-
-			FString GlobalName(GlobalEnd - GlobalStart, GlobalStart);
-			GlobalName.TrimStartAndEndInline();
-			GlobalName.ReplaceInline(TEXT(" "), TEXT(""));
+			FillUniformBufferDefinition(StructName, StructStart, StructEnd);
+			FillUniformBufferDefinition(GlobalName, GlobalStart, GlobalEnd);
 
 			// Avoid unnecessary conversions
-			if (StructName != GlobalName)
+			if (StructName.Len() == GlobalName.Len() && FCString::Strncmp(*StructName, *GlobalName, StructName.Len()) != 0)
 			{
 				FUniformBufferMemberInfo NewMemberInfo;
 				NewMemberInfo.NameAsStructMember = StructName;
 				NewMemberInfo.GlobalName = GlobalName;
 
 				// Add this member to the map
-				TArray<FUniformBufferMemberInfo>& UniformBufferMembers = UniformBufferNameToMembers.FindOrAdd(UniformBufferName);
-				UniformBufferMembers.Add(MoveTemp(NewMemberInfo));
+				if (FUniformBufferInfo* Info = UniformBufferInfos.Find(UniformBufferName))
+				{
+					Info->Members.Add(MoveTemp(NewMemberInfo));
+				}
+				else
+				{
+					FUniformBufferInfo NewInfo;
+					NewInfo.DefinitionEndOffset = ClosingBrace - SourceStart;
+					NewInfo.Members.Add(MoveTemp(NewMemberInfo));
+
+					UniformBufferInfos.Emplace(UniformBufferName, MoveTemp(NewInfo));
+				}
 			}
 		}
 
 		CurrentParseStart = NextSemicolon + 1;
-		NextSemicolon = FCString::Strstr(CurrentParseStart, TEXT(";"));
+		NextSemicolon = FindNextChar(CurrentParseStart, ';');
 	}
 
 	const TCHAR* EndPtr = ClosingBrace;
@@ -857,14 +888,32 @@ const TCHAR* ParseUniformBufferDefinition(const TCHAR* ReadStart, TMap<FString,
 	return EndPtr;
 }
 
+inline bool IsValidSymbolStart(const TCHAR* SearchPtr, int32 UniformBufferNameLen)
+{
+	// Make sure this isn't the end of another symbol or a sub-member of a different struct
+	//   ex: "OpaqueBasePass.Substrate.MaxBytesPerPixel" when matching "Substrate"
+	if (IsValidHLSLIdentifierCharacter(*(SearchPtr - 1)) || *(SearchPtr - 1) == '.')
+	{
+		return false;
+	}
+
+	// Avoid lines like '#line   1 "/Engine/Generated/UniformBuffers/View.ush"'
+	if (FCString::Strncmp(SearchPtr + UniformBufferNameLen, TEXT("ush"), 3) == 0)
+	{
+		return false;
+	}
+
+	return true;
+}
+
 // The cross compiler doesn't yet support struct initializers needed to construct static structs for uniform buffers
 // Replace all uniform buffer struct member references (View.WorldToClip) with a flattened name that removes the struct dependency (View_WorldToClip)
 void CleanupUniformBufferCode(const FShaderCompilerEnvironment& Environment, FString& PreprocessedShaderSource)
 {
 	TRACE_CPUPROFILER_EVENT_SCOPE(CleanupUniformBufferCode);
 
-	TMap<FString, TArray<FUniformBufferMemberInfo>> UniformBufferNameToMembers;
-	UniformBufferNameToMembers.Reserve(Environment.UniformBufferMap.Num());
+	TMap<FStringView, FUniformBufferInfo> UniformBufferInfos;
+	UniformBufferInfos.Reserve(Environment.UniformBufferMap.Num());
 
 	// Build a mapping from uniform buffer name to its members
 	{
@@ -877,9 +926,9 @@ void CleanupUniformBufferCode(const FShaderCompilerEnvironment& Environment, FSt
 		while (SearchPtr)
 		{
 			const ptrdiff_t Offset = SearchPtr - SourceStart;
-			if (Offset > 0 && FChar::IsWhitespace(*(SearchPtr - 1)) && FChar::IsWhitespace(*(SearchPtr + StructIdentifierLen)))
+			if (Offset > 0 && IsSpaceOrTabOrEOL(*(SearchPtr - 1)) && IsSpaceOrTabOrEOL(*(SearchPtr + StructIdentifierLen)))
 			{
-				const TCHAR* ConstStructEndPtr = ParseUniformBufferDefinition(SearchPtr + StructIdentifierLen, UniformBufferNameToMembers);
+				const TCHAR* ConstStructEndPtr = ParseUniformBufferDefinition(&PreprocessedShaderSource[0], SearchPtr + StructIdentifierLen, UniformBufferInfos);
 				TCHAR* StructEndPtr = &PreprocessedShaderSource[ConstStructEndPtr - &PreprocessedShaderSource[0]];
 
 				// Comment out the uniform buffer struct and initializer
@@ -897,59 +946,71 @@ void CleanupUniformBufferCode(const FShaderCompilerEnvironment& Environment, FSt
 		}
 	}
 
+	const TCHAR* const SearchSuffixes[] = { TEXT("."), TEXT(" .") };
+	TStringBuilder<64> UniformBufferName;
+
 	// Replace all uniform buffer struct member references (View.WorldToClip) with a flattened name that removes the struct dependency (View_WorldToClip)
-	for (TMap<FString, TArray<FUniformBufferMemberInfo>>::TConstIterator It(UniformBufferNameToMembers); It; ++It)
+	for (TMap<FStringView, FUniformBufferInfo>::TConstIterator It(UniformBufferInfos); It; ++It)
 	{
-		const FString& UniformBufferName = It.Key();
-		FString UniformBufferAccessString = UniformBufferName + TEXT(".");
+		const FUniformBufferInfo& UniformBufferInfo = It.Value();
 
-		// Search for the uniform buffer name first, as an optimization (instead of searching the entire source for every member)
-		TCHAR* SearchPtr = FindNextUniformBufferReference(&PreprocessedShaderSource[0], *UniformBufferName, UniformBufferName.Len());
+		// Start after the definition of this uniform buffer
+		TCHAR* const SearchStart = &PreprocessedShaderSource[UniformBufferInfo.DefinitionEndOffset];
 
-		while (SearchPtr)
+		// Trying "." and " ." separately can be faster than checking for each case as we iterate
+		for (const TCHAR* SearchSuffix : SearchSuffixes)
 		{
-			const TArray<FUniformBufferMemberInfo>& UniformBufferMembers = It.Value();
+			UniformBufferName.Reset();
+			UniformBufferName.Append(It.Key());
+			UniformBufferName.Append(SearchSuffix);
 
-			// Find the matching member we are replacing
-			for (int32 MemberIndex = 0; MemberIndex < UniformBufferMembers.Num(); MemberIndex++)
-			{
-				const FString& MemberNameAsStructMember = UniformBufferMembers[MemberIndex].NameAsStructMember;
+			// Search for the uniform buffer name first, as an optimization (instead of searching the entire source for every member)
+			TCHAR* SearchPtr = FCString::Strstr(SearchStart, *UniformBufferName);
 
-				if (MatchStructMemberName(MemberNameAsStructMember, SearchPtr, PreprocessedShaderSource))
+			while (SearchPtr)
+			{
+				// Only match whole symbol
+				if (IsValidSymbolStart(SearchPtr, UniformBufferName.Len()))
 				{
-					const FString& MemberNameGlobal = UniformBufferMembers[MemberIndex].GlobalName;
-					int32 NumWhitespacesToAdd = 0;
-
-					for (int32 i = 0; i < MemberNameAsStructMember.Len(); i++)
+					// Find the matching member we are replacing
+					for (const FUniformBufferMemberInfo& MemberInfo : UniformBufferInfo.Members)
 					{
-						if (i < MemberNameAsStructMember.Len() - 1)
+						if (const TCHAR* SearchEnd = MatchStructMemberName(MemberInfo.NameAsStructMember, SearchPtr))
 						{
-							if (FChar::IsWhitespace(SearchPtr[i]))
+							const int32 FoundTextLen = SearchEnd - SearchPtr;
+							const int32 ReplacementTextLen = MemberInfo.GlobalName.Len();
+
+							check(FoundTextLen >= ReplacementTextLen);
+
+							const TCHAR* SourceStart = GetData(MemberInfo.GlobalName);
+							TCHAR* const DestinationStart = SearchPtr;
+
+							int32 Index = 0;
+							for (; Index < ReplacementTextLen; Index++)
 							{
-								NumWhitespacesToAdd++;
+								DestinationStart[Index] = SourceStart[Index];
 							}
-						}
 
-						SearchPtr[i] = MemberNameGlobal[i];
-					}
+							// The shader preprocessor inserts spaces after defines
+							// #define ReflectionStruct OpaqueBasePass.Shared.Reflection
+							// 'ReflectionStruct.SkyLightCubemapBrightness' becomes 'OpaqueBasePass.Shared.Reflection .SkyLightCubemapBrightness' after MCPP
+							// In order to convert this struct member reference into a globally unique variable we move the spaces to the end
+							// 'OpaqueBasePass.Shared.Reflection .SkyLightCubemapBrightness' -> 'OpaqueBasePass_Shared_Reflection_SkyLightCubemapBrightness '
 
-					// The shader preprocessor inserts spaces after defines
-					// #define ReflectionStruct OpaqueBasePass.Shared.Reflection
-					// 'ReflectionStruct.SkyLightCubemapBrightness' becomes 'OpaqueBasePass.Shared.Reflection .SkyLightCubemapBrightness' after MCPP
-					// In order to convert this struct member reference into a globally unique variable we move the spaces to the end
-					// 'OpaqueBasePass.Shared.Reflection .SkyLightCubemapBrightness' -> 'OpaqueBasePass_Shared_Reflection_SkyLightCubemapBrightness '
-					for (int32 i = 0; i < NumWhitespacesToAdd; i++)
-					{
-						// If we passed MatchStructMemberName, it should not be possible to overwrite the null terminator
-						check(SearchPtr[MemberNameAsStructMember.Len() + i] != 0);
-						SearchPtr[MemberNameAsStructMember.Len() + i] = ' ';
+							for (; Index < FoundTextLen; Index++)
+							{
+								// If we passed MatchStructMemberName, it should not be possible to overwrite the null terminator
+								check(DestinationStart[Index] != '\0');
+								DestinationStart[Index] = ' ';
+							}
+
+							break;
+						}
 					}
-							
-					break;
 				}
-			}
 
-			SearchPtr = FindNextUniformBufferReference(SearchPtr + UniformBufferAccessString.Len(), *UniformBufferName, UniformBufferName.Len());
+				SearchPtr = FCString::Strstr(SearchPtr + UniformBufferName.Len(), *UniformBufferName);
+			}
 		}
 	}
 }
@@ -959,7 +1020,7 @@ void CleanupUniformBufferCode(const FShaderCompilerEnvironment& Environment, FSt
 
 static FString ParseText(const TCHAR* StartPtr, const TCHAR*& EndPtr)
 {
-	const TCHAR* OpeningBracePtr = FCString::Strstr(StartPtr, TEXT("("));
+	const TCHAR* OpeningBracePtr = FindNextChar(StartPtr, '(');
 	check(OpeningBracePtr);
 
 	const TCHAR* ClosingBracePtr = FindMatchingClosingParenthesis(OpeningBracePtr + 1);
