Shader Compiler:  Optimized version of CleanupUniformBufferCode, called during preprocessing.  New version is 3.8x faster, producing a 19% overall speedup in ConditionalPreprocessShader.  Processes the shader in a single pass, scanning for potential compound identifiers (symbol tokens connected by dots and optional whitespace), then checking if they match known uniform buffer members.  A few optimizations improve the speed of checking if an identifier matches, including an early out based on identifier length and start character, followed by organizing uniform buffers and their members by length, requiring only a small subset of members to be tested.  Parsing uniform buffer definitions was also optimized to avoid memory allocation.

Doesn't produce bit identical preprocessed output, as it compacts whitespace from compound identifiers before testing them, whereas the original implementation compacts whitespace after matching.  Compacting first makes testing more efficient and simpler, since it doesn't need to take into account whitespace.  For A/B validation, I did add temporary debug code that stored and reverted un-matched identifiers, mimicking the old behavior, and it was identical across the board.

#rnx

diff --git a/Engine/Source/Developer/ShaderCompilerCommon/Private/ShaderCompilerCommon.cpp b/Engine/Source/Developer/ShaderCompilerCommon/Private/ShaderCompilerCommon.cpp
--- a/Engine/Source/Developer/ShaderCompilerCommon/Private/ShaderCompilerCommon.cpp
+++ b/Engine/Source/Developer/ShaderCompilerCommon/Private/ShaderCompilerCommon.cpp
@@ -12,6 +12,14 @@
 #include "ShaderPreprocessTypes.h"
 #include "ShaderSymbolExport.h"
 #include "ShaderMinifier.h"
+#include "Algo/Sort.h"
+
+static TAutoConsoleVariable<bool> CVarShaderCompilerCleanupUniformBufferCodeNew(
+	TEXT("r.ShaderCompiler.CleanupUniformBufferCodeNew"),
+	true,
+	TEXT("Run new optimized version of CleanupUniformBufferCode.  Temporary chicken switch in case there are issues, before we remove old version."),
+	ECVF_Default
+);
 
 IMPLEMENT_MODULE(FDefaultModuleImpl, ShaderCompilerCommon);
 
@@ -806,6 +814,41 @@ struct FUniformBufferInfo
 	TArray<FUniformBufferMemberInfo> Members;
 };
 
+struct FUniformBufferMemberInfoNew
+{
+	// eg View.WorldToClip
+	FStringView NameAsStructMember;
+	// eg View_WorldToClip
+	FStringView GlobalName;
+
+	bool operator<(const FUniformBufferMemberInfoNew& Other)
+	{
+		if (NameAsStructMember.Len() != Other.NameAsStructMember.Len())
+		{
+			return NameAsStructMember.Len() < Other.NameAsStructMember.Len();
+		}
+		else
+		{
+			return NameAsStructMember.Compare(Other.NameAsStructMember, ESearchCase::CaseSensitive) < 0;
+		}
+	}
+};
+
+// Index and count of subset of members
+struct FUniformBufferMemberView
+{
+	int32 MemberOffset;
+	int32 MemberCount;
+};
+
+struct FUniformBufferInfoNew
+{
+	FStringView Name;
+	int32 NextWithSameLength;							// Linked list of uniform buffer infos with same name length
+	TArray<FUniformBufferMemberInfoNew> Members;		// Members sorted by length
+	TArray<FUniformBufferMemberView> MembersByLength;	// Offset and count of members of a given length
+};
+
 static void FillUniformBufferDefinition(TStringBuilder<128>& Builder, const TCHAR* Start, const TCHAR* End)
 {
 	Builder.Reset();
@@ -906,12 +949,560 @@ inline bool IsValidSymbolStart(const TCHAR* SearchPtr, int32 UniformBufferNameLe
 	return true;
 }
 
+// Compacts spaces out of a compound identifier.  Returns the new end pointer of the compacted identifier.
+// End and result pointers are exclusive (length of the string is End - Start).
+static TCHAR* CompactCompoundIdentifier(TCHAR* Start, TCHAR* End)
+{
+	// Find first whitespace in the identifier, if present
+	TCHAR* ReadChar;
+	for (ReadChar = Start; ReadChar < End; ++ReadChar)
+	{
+		if (IsSpaceOrTabOrEOL(*ReadChar))
+		{
+			break;
+		}
+	}
+	if (ReadChar == End)
+	{
+		// No whitespace, we're done!
+		return End;
+	}
+
+	// Found some whitespace, so we need to compact the non-whitespace, swapping the whitespace to the end of the range
+	// WriteChar here will be the first whitespace character that we need to compact into.
+	TCHAR* WriteChar = ReadChar;
+	for (++ReadChar; ReadChar < End; ++ReadChar)
+	{
+		// If the current read character is non-whitespace, compact it down
+		if (!IsSpaceOrTabOrEOL(*ReadChar))
+		{
+			Swap(*ReadChar, *WriteChar);
+			WriteChar++;
+		}
+	}
+	return WriteChar;
+}
+
+const TCHAR* ParseUniformBufferDefinitionNew(const TCHAR* ReadStart, TArray<FUniformBufferInfoNew>& UniformBufferInfos, uint64 UniformBufferFilter[64], int32 UniformBuffersByLength[64])
+{
+	// TODO:  should we check for an existing item?  In my testing, there's only one uniform buffer declaration with a given name,
+	// but the original code used a map, theoretically allowing for multiple.
+	int32 InfoIndex = UniformBufferInfos.AddDefaulted();
+	FUniformBufferInfoNew& Info = UniformBufferInfos[InfoIndex];
+
+	Info.Name = ParseHLSLSymbolName(ReadStart);
+	check(Info.Name.Len() < 64);
+
+	const TCHAR* OpeningBrace = FindNextChar(ReadStart, '{');
+	const TCHAR* ClosingBrace = FindMatchingClosingBrace(OpeningBrace + 1);
+
+	const TCHAR* CurrentParseStart = OpeningBrace + 1;
+	const TCHAR* NextSemicolon = FindNextChar(CurrentParseStart, ';');
+
+	while (NextSemicolon < ClosingBrace)
+	{
+		const TCHAR* NextSeparator = FindNextChar(CurrentParseStart, '=');
+		if (NextSeparator < NextSemicolon)
+		{
+			const TCHAR* StructStart = CurrentParseStart;
+			const TCHAR* StructEnd = NextSeparator - 1;
+
+			const TCHAR* GlobalStart = NextSeparator + 1;
+			const TCHAR* GlobalEnd = NextSemicolon - 1;
+
+			while (IsSpaceOrTabOrEOL(*StructStart))
+			{
+				StructStart++;
+			}
+			while (IsSpaceOrTabOrEOL(*GlobalStart))
+			{
+				GlobalStart++;
+			}
+
+			StructEnd = CompactCompoundIdentifier(const_cast<TCHAR*>(StructStart), const_cast<TCHAR*>(StructEnd));
+			GlobalEnd = CompactCompoundIdentifier(const_cast<TCHAR*>(GlobalStart), const_cast<TCHAR*>(GlobalEnd));
+
+			FStringView StructName(StructStart, StructEnd - StructStart);
+			FStringView GlobalName(GlobalStart, GlobalEnd - GlobalStart);
+
+			// Avoid unnecessary conversions
+			if (StructName.Len() == GlobalName.Len() && FCString::Strncmp(StructName.GetData(), GlobalName.GetData(), StructName.Len()) != 0)
+			{
+				FUniformBufferMemberInfoNew NewMemberInfo;
+				NewMemberInfo.NameAsStructMember = StructName;
+				NewMemberInfo.GlobalName = GlobalName;
+
+				// Need to be able to replace strings in place, so make sure GlobalName will fit in space of NameAsStructMember
+				check(NewMemberInfo.NameAsStructMember.Len() >= NewMemberInfo.GlobalName.Len());
+
+				Info.Members.Add(NewMemberInfo);
+			}
+		}
+
+		CurrentParseStart = NextSemicolon + 1;
+		NextSemicolon = FindNextChar(CurrentParseStart, ';');
+	}
+
+	const TCHAR* EndPtr = ClosingBrace;
+
+	// Skip to the end of the UniformBuffer
+	while (*EndPtr && *EndPtr != ';')
+	{
+		EndPtr++;
+	}
+
+	if (Info.Members.Num())
+	{
+		// We have members.  Sort them.  Note that the sort is by length first, not alphabetical, so the last item will be the longest.
+		Algo::Sort(Info.Members);
+
+		int32 MaxLen = Info.Members.Last().NameAsStructMember.Len();
+
+		// Initialize table with offset of first member with a given length, and the count of members of that length (going backwards so the
+		// index of the first element of a given size is the last one written to "MemberOffset").
+		Info.MembersByLength.SetNumZeroed(MaxLen + 1);
+
+		for (int32 MemberIndex = Info.Members.Num() - 1; MemberIndex >= 0; MemberIndex--)
+		{
+			int32 CurrentMemberLen = Info.Members[MemberIndex].NameAsStructMember.Len();
+			Info.MembersByLength[CurrentMemberLen].MemberOffset = MemberIndex;
+			Info.MembersByLength[CurrentMemberLen].MemberCount++;
+		}
+
+		// Initialize the uniform buffer name filter.  The filter is a mask based on the first character of the name (minus 64 so valid token
+		// starting characters which are in ASCII range 64..127 fit in 64 bits).  We can quickly check if a token of the given length and start
+		// character might be one we care about.
+		UniformBufferFilter[Info.Name.Len()] |= 1ull << (Info.Name[0] - 64);
+
+		// Add to linked list of uniform buffers by name length
+		Info.NextWithSameLength = UniformBuffersByLength[Info.Name.Len()];
+		UniformBuffersByLength[Info.Name.Len()] = InfoIndex;
+	}
+	else
+	{
+		// If no members, we don't care about it
+		UniformBufferInfos.RemoveAt(UniformBufferInfos.Num() - 1);
+	}
+
+	return EndPtr;
+}
+
+enum class AsciiFlags
+{
+	TerminatorOrSlash = (1 << 0),	// Null terminator OR slash (latter we care about for skipping commented out uniform blocks)
+	Whitespace = (1 << 1),			// Includes other special characters below 32 (in addition to tab / newline)
+	Other = (1 << 2),				// Anything else not one of the other types
+	SymbolStart = (1<<3),			// Letters plus underscore (anything that can start a symbol)
+	Digit = (1 << 4),
+	Dot = (1 << 5),
+	Quote = (1 << 6),
+	Hash = (1 << 7),
+};
+
+static uint8 AsciiFlagTable[256] =
+{
+	1,2,2,2,2,2,2,2, 2,2,2,2,2,2,2,2, 2,2,2,2,2,2,2,2, 2,2,2,2,2,2,2,2,		// Treat all special characters as whitespace
+
+	2,4,64,128,4,4,4,4,			// 34 == Quote  35 == Hash
+	4,4,4,4,4,4,32,1,			// 46 == Dot    47 == Slash
+	16,16,16,16,16,16,16,16,	// Digits 0-7
+	16,16,4,4,4,4,4,4,			// Digits 8-9
+
+	4,8,8,8,8,8,8,8, 8,8,8,8,8,8,8,8, 8,8,8,8,8,8,8,8, 8,8,8,4,4,4,4,8,		// Upper case letters,  95 == Underscore
+	4,8,8,8,8,8,8,8, 8,8,8,8,8,8,8,8, 8,8,8,8,8,8,8,8, 8,8,8,4,4,4,4,4,		// Lower case letters
+
+	4,4,4,4,4,4,4,4, 4,4,4,4,4,4,4,4, 4,4,4,4,4,4,4,4, 4,4,4,4,4,4,4,4,		// Treat all non-ASCII characters as Other
+	4,4,4,4,4,4,4,4, 4,4,4,4,4,4,4,4, 4,4,4,4,4,4,4,4, 4,4,4,4,4,4,4,4,
+	4,4,4,4,4,4,4,4, 4,4,4,4,4,4,4,4, 4,4,4,4,4,4,4,4, 4,4,4,4,4,4,4,4,
+	4,4,4,4,4,4,4,4, 4,4,4,4,4,4,4,4, 4,4,4,4,4,4,4,4, 4,4,4,4,4,4,4,4,
+};
+
+struct FCompoundIdentifierResult
+{
+	const TCHAR* Identifier;			// Start of identifier
+	const TCHAR* IdentifierEnd;			// End of entire identifier
+	const TCHAR* IdentifierRootEnd;		// End of root token of identifier
+};
+
+// Searches for a "compound identifier" (series of symbol tokens separated by dots) that also passes the "RootIdentifierFilter".
+// The filter is a mask table of valid identifier start characters indexed by identifier length.  Since identifier characters start
+// with letters or underscore, we can store a 64-bit mask representing ASCII characters 64..127, as all valid start characters are
+// in that range.  As an example, if "View" is a valid root identifier, RootIdentifierFilter[4] will have the bit ('V' - 64) set,
+// and any other 4 character identifier that doesn't start with that letter can be skipped, saving overhead in the caller.
+bool FindNextCompoundIdentifier(const TCHAR*& Search, const uint64 RootIdentifierFilter[64], FCompoundIdentifierResult& OutResult)
+{
+	const TCHAR* SearchChar = Search;
+	uint8 SearchCharFlag = AsciiFlagTable[(uint8)*SearchChar];
+
+	// Scanning loop
+	while (1)
+	{
+		static constexpr uint8 AsciiFlagsEchoVerbatim = (uint8)AsciiFlags::Whitespace | (uint8)AsciiFlags::Other;
+		static constexpr uint8 AsciiFlagsSymbol = (uint8)AsciiFlags::SymbolStart | (uint8)AsciiFlags::Digit;
+		static constexpr uint8 AsciiFlagsStartNumberOrDirective = (uint8)AsciiFlags::Digit | (uint8)AsciiFlags::Dot | (uint8)AsciiFlags::Hash;
+		static constexpr uint8 AsciiFlagsEndNumberOrDirective = (uint8)AsciiFlags::Whitespace | (uint8)AsciiFlags::Other | (uint8)AsciiFlags::Quote | (uint8)AsciiFlags::TerminatorOrSlash;
+
+		// Conditions here are organized in expected order of frequency
+		if (SearchCharFlag & AsciiFlagsEchoVerbatim)
+		{
+			SearchChar++;
+			SearchCharFlag = AsciiFlagTable[(uint8)*SearchChar];
+		}
+		else if (SearchCharFlag & (uint8)AsciiFlags::SymbolStart)
+		{
+			OutResult.Identifier = SearchChar;
+			SearchChar++;
+			while ((SearchCharFlag = AsciiFlagTable[(uint8)*SearchChar]) & AsciiFlagsSymbol)
+			{
+				SearchChar++;
+			}
+
+			// Track end of our root identifier
+			OutResult.IdentifierRootEnd = SearchChar;
+
+			// Skip any whitespace before a potential dot
+			while (SearchCharFlag & ((uint8)AsciiFlags::Whitespace))
+			{
+				SearchChar++;
+				SearchCharFlag = AsciiFlagTable[(uint8)*SearchChar];
+			}
+
+			// If we didn't find a dot, go back to initial scanning state
+			if (!(SearchCharFlag & ((uint8)AsciiFlags::Dot)))
+			{
+				continue;
+			}
+			SearchChar++;
+			SearchCharFlag = AsciiFlagTable[(uint8)*SearchChar];
+
+			// Determine if this root identifier passes the filter.  If so, we'll continue to parse the rest of the identifier,
+			// but then go back to scanning.  The mask in RootIdentifierFilter starts with ASCII character 64, as token start
+			// characters are in the range [64..127].
+			ptrdiff_t IdentifierRootLen = OutResult.IdentifierRootEnd - OutResult.Identifier;
+			if (IdentifierRootLen >= 64 || !(RootIdentifierFilter[IdentifierRootLen] & (1ull << (*OutResult.Identifier - 64))))
+			{
+				// Clear this, marking that we didn't find a candidate root identifier
+				OutResult.IdentifierRootEnd = nullptr;
+			}
+
+			// Skip any whitespace after dot
+			while (SearchCharFlag & ((uint8)AsciiFlags::Whitespace))
+			{
+				SearchChar++;
+				SearchCharFlag = AsciiFlagTable[(uint8)*SearchChar];
+			}
+
+			// Check for the start of another symbol after the dot -- if it's not a symbol, switch back to scanning -- some kind of incorrect code
+			if (!(SearchCharFlag & (uint8)AsciiFlags::SymbolStart))
+			{
+				continue;
+			}
+
+			// Repeatedly scan for additional parts of the identifier separated by dots
+			while (1)
+			{
+				SearchChar++;
+				while ((SearchCharFlag = AsciiFlagTable[(uint8)*SearchChar]) & AsciiFlagsSymbol)
+				{
+					SearchChar++;
+				}
+
+				// Track that this may be the end of the identifier (if there's not more dot separated tokens)
+				OutResult.IdentifierEnd = SearchChar;
+
+				// Skip any whitespace before a potential dot
+				while (SearchCharFlag & ((uint8)AsciiFlags::Whitespace))
+				{
+					SearchChar++;
+					SearchCharFlag = AsciiFlagTable[(uint8)*SearchChar];
+				}
+
+				// If we found something other than a dot, we're done!
+				if (!(SearchCharFlag & ((uint8)AsciiFlags::Dot)))
+				{
+					// Is the root token for this identifier a candidate based on the filter?
+					if (OutResult.IdentifierRootEnd)
+					{
+						Search = SearchChar;
+						return true;
+					}
+					else
+					{
+						// If not, go back to initial scanning state
+						break;
+					}
+				}
+
+				// Skip the dot
+				SearchChar++;
+				SearchCharFlag = AsciiFlagTable[(uint8)*SearchChar];
+
+				// Skip any whitespace after dot
+				while (SearchCharFlag & ((uint8)AsciiFlags::Whitespace))
+				{
+					SearchChar++;
+					SearchCharFlag = AsciiFlagTable[(uint8)*SearchChar];
+				}
+
+				// Did we find the start of another symbol after the dot?  If not, break out, some kind of invalid code...
+				if (!(SearchCharFlag & (uint8)AsciiFlags::SymbolStart))
+				{
+					break;
+				}
+			}
+		}
+		else if (SearchCharFlag & AsciiFlagsStartNumberOrDirective)
+		{
+			// Number or directive, skip to Whitespace, Other, or Quote (numbers may contain letters or #, i.e. "1.#INF" for infinity, or "e" for an exponent)
+			SearchChar++;
+			while (!((SearchCharFlag = AsciiFlagTable[(uint8)*SearchChar]) & AsciiFlagsEndNumberOrDirective))
+			{
+				SearchChar++;
+			}
+		}
+		else if (SearchCharFlag & (uint8)AsciiFlags::Quote)
+		{
+			// Quote, skip to next Quote (or maybe end of string if text is malformed)
+			SearchChar++;
+			while (*SearchChar && *SearchChar != TEXT('\"'))
+			{
+				SearchChar++;
+			}
+
+			// Could be end of string or the quote -- skip over the quote if not the null terminator
+			if (*SearchChar)
+			{
+				SearchChar++;
+			}
+			SearchCharFlag = AsciiFlagTable[(uint8)*SearchChar];
+		}
+		// Must be null terminator or slash at this point -- we've tested all other possibilities
+		else if (*SearchChar == TEXT('/'))
+		{
+			// Check if this is a commented out block (typically a commented out uniform declaration) and skip over it.
+			// If the text is bad, there could be a /* right at the end of the string, so we need to check there is at least
+			// one more character.
+			if (SearchChar[1] == TEXT('*') && SearchChar[2] != 0)
+			{
+				// Search for slash (or end of string), starting at SearchChar + 3.  If we find a slash, we'll check the previous
+				// character to see if it's the end of the comment.  Starting at +3 is necessary to avoid matching a slash as the
+				// first character of the comment, i.e. "/*/".
+				SearchChar += 3;
+
+				while (1)
+				{
+					while ((SearchCharFlag = AsciiFlagTable[(uint8)*SearchChar]) != (uint8)AsciiFlags::TerminatorOrSlash)
+					{
+						SearchChar++;
+					}
+
+					// Is this the end of the comment?
+					if (*(SearchChar - 1) == TEXT('*'))
+					{
+						if (*SearchChar)
+						{
+							SearchChar++;
+							SearchCharFlag = AsciiFlagTable[(uint8)*SearchChar];
+							break;
+						}
+					}
+					else
+					{
+						// More characters, continue the comment scanning loop, or if somehow at end of string, return false...
+						if (*SearchChar)
+						{
+							SearchChar++;
+						}
+						else
+						{
+							return false;
+						}
+					}
+				}
+			}
+			else
+			{
+				// Just a slash, not part of a block comment
+				SearchChar++;
+				SearchCharFlag = AsciiFlagTable[(uint8)*SearchChar];
+			}
+		}
+		else
+		{
+			// End of string
+			Search = SearchChar;
+			return false;
+		}
+	}
+}
+
+TCHAR* FindNextUniformBufferDefinition(TCHAR* SearchPtr, TCHAR* SourceStart, const TCHAR* UniformBufferStructIdentifier, int32 StructIdentifierLen)
+{
+	while (SearchPtr)
+	{
+		SearchPtr = FCString::Strstr(SearchPtr, UniformBufferStructIdentifier);
+
+		if (SearchPtr)
+		{
+			if (SearchPtr > SourceStart && IsSpaceOrTabOrEOL(*(SearchPtr - 1)) && IsSpaceOrTabOrEOL(*(SearchPtr + StructIdentifierLen)))
+			{
+				break;
+			}
+			else
+			{
+				SearchPtr = SearchPtr + 1;
+			}
+		}
+	}
+	return SearchPtr;
+}
+
+const TCHAR* FindPreviousDot(const TCHAR* SearchPtr, const TCHAR* SearchMin)
+{
+	while ((SearchPtr > SearchMin) && (*SearchPtr != TEXT('.')))
+	{
+		SearchPtr--;
+	}
+	return SearchPtr;
+}
+
+// Unlike the original version, this parses uniform buffers and replaces symbols in a single pass.
+void CleanupUniformBufferCodeNew(const FShaderCompilerEnvironment& Environment, FString& PreprocessedShaderSource)
+{
+	TArray<FUniformBufferInfoNew> UniformBufferInfos;
+	uint64 UniformBufferFilter[64] = { 0 };			// A bit set for valid start characters for uniform buffer name of given length
+	int32 UniformBuffersByLength[64];				// Linked list head index into UniformBufferInfos by length (connected by "NextWithSameLength")
+
+	UniformBufferInfos.Reserve(Environment.UniformBufferMap.Num());
+	memset(UniformBuffersByLength, 0xff, sizeof(UniformBuffersByLength));
+
+	const TCHAR* UniformBufferStructIdentifier = TEXT("UniformBuffer");
+	const int32 StructIdentifierLen = FCString::Strlen(UniformBufferStructIdentifier);
+
+	TCHAR* SourceStart = &PreprocessedShaderSource[0];
+	TCHAR* SearchPtr = SourceStart;
+	TCHAR* EndOfPreviousUniformBuffer = SourceStart;
+	bool bUniformBufferFound;
+
+	do
+	{
+		// Find the next uniform buffer definition
+		SearchPtr = FindNextUniformBufferDefinition(SearchPtr, SourceStart, UniformBufferStructIdentifier, StructIdentifierLen);
+
+		if (SearchPtr)
+		{
+			// Track that we found a uniform buffer, and temporarily null terminate the string so we can parse to this point
+			bUniformBufferFound = true;
+			*SearchPtr = 0;
+		}
+		else
+		{
+			bUniformBufferFound = false;
+		}
+
+		// Parse the source between the last uniform buffer and the current uniform buffer (or potentially the end of the source if no more
+		// were found).  If there are no uniform buffers yet, we don't need to parse anything.
+		if (UniformBufferInfos.Num())
+		{
+			const TCHAR* ParsePtr = EndOfPreviousUniformBuffer;
+
+			FCompoundIdentifierResult Result;
+			while (FindNextCompoundIdentifier(ParsePtr, UniformBufferFilter, Result))
+			{
+				// Check if the identifier corresponds to a uniform buffer
+				FStringView IdentifierRoot(Result.Identifier, Result.IdentifierRootEnd - Result.Identifier);
+				for (int32 UniformInfoIndex = UniformBuffersByLength[IdentifierRoot.Len()]; UniformInfoIndex != INDEX_NONE; UniformInfoIndex = UniformBufferInfos[UniformInfoIndex].NextWithSameLength)
+				{
+					FUniformBufferInfoNew& Info = UniformBufferInfos[UniformInfoIndex];
+					if (IdentifierRoot.Equals(Info.Name, ESearchCase::CaseSensitive))
+					{
+						// Found the uniform buffer, clean up potential whitespace
+						Result.IdentifierEnd = CompactCompoundIdentifier(const_cast<TCHAR*>(Result.Identifier), const_cast<TCHAR*>(Result.IdentifierEnd));
+
+						// Now try to find a matching member.  We need to check subsets of the full "identifier", to strip away function calls, components, or child structures.
+						bool bMatchFound = false;
+
+						for (; Result.IdentifierEnd > Result.IdentifierRootEnd; Result.IdentifierEnd = FindPreviousDot(Result.IdentifierEnd - 1, Result.IdentifierRootEnd))
+						{
+							FStringView Identifier(Result.Identifier, Result.IdentifierEnd - Result.Identifier);
+							if (Identifier.Len() < Info.MembersByLength.Num())
+							{
+								const FUniformBufferMemberView& MemberView = Info.MembersByLength[Identifier.Len()];
+
+								for (int32 MemberIndex = MemberView.MemberOffset; MemberIndex < MemberView.MemberOffset + MemberView.MemberCount; MemberIndex++)
+								{
+									if (Info.Members[MemberIndex].NameAsStructMember.Equals(Identifier, ESearchCase::CaseSensitive))
+									{
+										bMatchFound = true;
+
+										const int32 OriginalTextLen = Info.Members[MemberIndex].NameAsStructMember.Len();
+										const int32 ReplacementTextLen = Info.Members[MemberIndex].GlobalName.Len();
+
+										const TCHAR* GlobalNameStart = GetData(Info.Members[MemberIndex].GlobalName);
+										TCHAR* IdentifierStart = const_cast<TCHAR*>(Result.Identifier);
+
+										int32 Index = 0;
+										for (; Index < ReplacementTextLen; Index++)
+										{
+											IdentifierStart[Index] = GlobalNameStart[Index];
+										}
+										for (; Index < OriginalTextLen; Index++)
+										{
+											IdentifierStart[Index] = ' ';
+										}
+										break;
+									}
+								}
+
+								if (bMatchFound)
+								{
+									break;
+								}
+							}
+						}
+
+						break;
+					}
+				}
+			}
+		}
+
+		// Parse the current uniform buffer.
+		if (bUniformBufferFound)
+		{
+			// Unterminate the string (put the first character of the struct identifier back in place) and parse it
+			*SearchPtr = UniformBufferStructIdentifier[0];
+
+			const TCHAR* ConstStructEndPtr = ParseUniformBufferDefinitionNew(SearchPtr + StructIdentifierLen, UniformBufferInfos, UniformBufferFilter, UniformBuffersByLength);
+			TCHAR* StructEndPtr = &PreprocessedShaderSource[ConstStructEndPtr - &PreprocessedShaderSource[0]];
+
+			// Comment out the uniform buffer struct and initializer
+			*SearchPtr = '/';
+			*(SearchPtr + 1) = '*';
+			*(StructEndPtr - 1) = '*';
+			*StructEndPtr = '/';
+
+			EndOfPreviousUniformBuffer = StructEndPtr + 1;
+			SearchPtr = StructEndPtr + 1;
+		}
+
+	} while (bUniformBufferFound);
+}
+
 // The cross compiler doesn't yet support struct initializers needed to construct static structs for uniform buffers
 // Replace all uniform buffer struct member references (View.WorldToClip) with a flattened name that removes the struct dependency (View_WorldToClip)
 void CleanupUniformBufferCode(const FShaderCompilerEnvironment& Environment, FString& PreprocessedShaderSource)
 {
 	TRACE_CPUPROFILER_EVENT_SCOPE(CleanupUniformBufferCode);
 
+	// Putting new code on a chicken switch in case there are issues
+	static const bool bRunCleanupUniformBufferCodeNew = CVarShaderCompilerCleanupUniformBufferCodeNew.GetValueOnAnyThread();
+	if (bRunCleanupUniformBufferCodeNew)
+	{
+		CleanupUniformBufferCodeNew(Environment, PreprocessedShaderSource);
+		return;
+	}
+
 	TMap<FStringView, FUniformBufferInfo> UniformBufferInfos;
 	UniformBufferInfos.Reserve(Environment.UniformBufferMap.Num());
 
