Adding typedef support to ShaderMinifier.
- Before this, every typedef and the code using them was being included in the final source.
- This was primarily affecting Bindless Resources where all bindless resources were included in the minified code even if they weren't used.

diff --git a/Engine/Source/Developer/ShaderCompilerCommon/Private/ShaderMinifier.cpp b/Engine/Source/Developer/ShaderCompilerCommon/Private/ShaderMinifier.cpp
--- a/Engine/Source/Developer/ShaderCompilerCommon/Private/ShaderMinifier.cpp
+++ b/Engine/Source/Developer/ShaderCompilerCommon/Private/ShaderMinifier.cpp
@@ -243,6 +243,7 @@ enum class ECodeChunkType {
 	CommentLine, // Single line comment
 	Namespace,
 	Using,
+	Typedef,
 };
 
 struct FNamespace
@@ -498,6 +499,15 @@ static FParsedShader ParseShader(FStringView InSource, FDiagnostics& Output)
 			{
 				// TODO: tag name / type / binding
 			}
+			else if (ChunkType == ECodeChunkType::Typedef)
+			{
+				NameBlockIndex = PendingBlocks.Num() - 1;
+				for (int32 Index = 1; Index < NameBlockIndex; Index++)
+				{
+					PendingBlocks[Index].Type = EBlockType::Type;
+				}
+				PendingBlocks[NameBlockIndex].Type = EBlockType::Name;
+			}
 
 			if (ChunkType == ECodeChunkType::Struct && bHasName && !bHasType)
 			{
@@ -705,6 +715,13 @@ static FParsedShader ParseShader(FStringView InSource, FDiagnostics& Output)
 					AddBlock(EBlockType::Keyword, Identifier);
 					continue;
 				}
+				else if (Identifier == TEXT("typedef"))
+				{
+					ChunkType = ECodeChunkType::Typedef;
+					Source = Remainder;
+					AddBlock(EBlockType::Keyword, Identifier);
+					continue;
+				}
 			}
 			else if (ChunkType == ECodeChunkType::Namespace)
 			{
@@ -1268,6 +1285,7 @@ static FString MinifyShader(const FParsedShader& Parsed, TConstArrayView<FString
 		ProcessedIdentifiers.Add(TEXT("uint4"));
 		ProcessedIdentifiers.Add(TEXT("void"));
 		ProcessedIdentifiers.Add(TEXT("while"));
+		ProcessedIdentifiers.Add(TEXT("typedef"));
 
 		// HLSL resource types
 		ProcessedIdentifiers.Add(TEXT("TextureCubeArray"));
@@ -1378,6 +1396,11 @@ static FString MinifyShader(const FParsedShader& Parsed, TConstArrayView<FString
 				continue;
 			}
 
+			if (Chunk.Type == ECodeChunkType::Typedef && Block.Type != EBlockType::Name)
+			{
+				continue;
+			}
+
 			if (Chunk.Type == ECodeChunkType::CBuffer && Block.Type == EBlockType::Body)
 			{
 				TempIdentifiers.Empty();
@@ -1946,6 +1969,23 @@ bool FShaderMinifierParserTest::RunTest(const FString& Parameters)
 		}
 	}
 
+	{
+		auto P = ParseShader(TEXT("typedef Bar Foo;"));
+		if (TestEqual(TEXT("ParseShader: standard typedef: num chunks"), P.Chunks.Num(), 1))
+		{
+			TestEqual(TEXT("ParseShader: standard typedef: chunk type"), P.Chunks[0].Type, ECodeChunkType::Typedef);
+		}
+	}
+
+	{
+		auto P = ParseShader(TEXT("typedef Bar Foo; static const Foo = Bar(0);"));
+		if (TestEqual(TEXT("ParseShader: standard typedef: num chunks"), P.Chunks.Num(), 2))
+		{
+			TestEqual(TEXT("ParseShader: standard typedef: chunk type"), P.Chunks[0].Type, ECodeChunkType::Typedef);
+			TestEqual(TEXT("ParseShader: standard typedef: chunk type"), P.Chunks[1].Type, ECodeChunkType::Variable);
+		}
+	}
+
 	int32 NumErrors = ExecutionInfo.GetErrorTotal();
 
 	return NumErrors == 0;
@@ -1983,6 +2023,8 @@ uint UnreferencedFunction()
 	return GUnreferencedParameter;
 }
 
+typedef Texture2D<float4> UnreferencedTypedef;
+
 #define COMPILER_DEFINITION_TEST 123
 float Sum(in FBar Param)
 {
@@ -2044,7 +2086,24 @@ static const struct
 	int Foo;
 } GInitializedAnonymousStructB = { 123 };
 
-RWBuffer<float4> OutputBuffer;
+typedef RWBuffer<float4> OutputBufferType;
+OutputBufferType OutputBuffer;
+
+struct FTypedefUsedStruct
+{
+	float Foo;
+};
+typedef StructuredBuffer<FTypedefUsedStruct> FTypedefUsed;
+typedef FTypedefUsed FTypedefUsedChained;
+typedef FTypedefUsedChained FTypedefUsedChainedUnused;
+FTypedefUsedChained TypedefUsedBuffer;
+
+struct FTypedefUnusedStruct
+{
+	float Foo;
+};
+typedef StructuredBuffer<FTypedefUnusedStruct> FTypedefUnused;
+FTypedefUnused TypedefUnusedBuffer;
 
 // Test comment 2
 [numthreads(1,1,1)]
@@ -2056,7 +2115,8 @@ void MainCS()
 	float A = FunB(GAnonymousStruct.Foo);
 	float B = FunB(GStructA.Bar + GStructB.Foo + GStructC.Foo);
 	float C = FunB(GInitializedAnonymousStructA.Foo + GInitializedAnonymousStructB.Foo);
-	OutputBuffer[0] = A + B;
+	float D = TypedefUsedBuffer[0].Foo;
+	OutputBuffer[0] = A + B + D;
 }
 )");
 
@@ -2107,12 +2167,21 @@ void MainCS()
 		TestTrue(TEXT("MinifyShader: MainCS: contains GStructC"), ChunkPresent(MinifiedParsed, TEXT("GStructC")));
 		TestTrue(TEXT("MinifyShader: MainCS: contains GInitializedAnonymousStructA"), ChunkPresent(MinifiedParsed, TEXT("GInitializedAnonymousStructA")));
 		TestTrue(TEXT("MinifyShader: MainCS: contains GInitializedAnonymousStructB"), ChunkPresent(MinifiedParsed, TEXT("GInitializedAnonymousStructB")));
+		TestTrue(TEXT("MinifyShader: MainCS: contains OutputBufferType"), ChunkPresent(MinifiedParsed, TEXT("OutputBufferType")));
 		TestTrue(TEXT("MinifyShader: MainCS: contains OutputBuffer"), ChunkPresent(MinifiedParsed, TEXT("OutputBuffer")));
+		TestTrue(TEXT("MinifyShader: MainCS: contains FTypedefUsedStruct"), ChunkPresent(MinifiedParsed, TEXT("FTypedefUsedStruct")));
+		TestTrue(TEXT("MinifyShader: MainCS: contains FTypedefUsed"), ChunkPresent(MinifiedParsed, TEXT("FTypedefUsed")));
+		TestTrue(TEXT("MinifyShader: MainCS: contains FTypedefUsedChained"), ChunkPresent(MinifiedParsed, TEXT("FTypedefUsedChained")));
+		TestTrue(TEXT("MinifyShader: MainCS: contains TypedefUsedBuffer"), ChunkPresent(MinifiedParsed, TEXT("TypedefUsedBuffer")));
 
 		// Expect false:
 		TestFalse(TEXT("MinifyShader: MainCS: contains UnreferencedFunction"), ChunkPresent(MinifiedParsed, TEXT("UnreferencedFunction")));
 		TestFalse(TEXT("MinifyShader: MainCS: contains FUnreferencedStruct"), ChunkPresent(MinifiedParsed, TEXT("FUnreferencedStruct")));
 		TestFalse(TEXT("MinifyShader: MainCS: contains GUnreferencedParameter"), ChunkPresent(MinifiedParsed, TEXT("GUnreferencedParameter")));
+		TestFalse(TEXT("MinifyShader: MainCS: contains FTypedefUsedChainedUnused"), ChunkPresent(MinifiedParsed, TEXT("FTypedefUsedChainedUnused")));
+		TestFalse(TEXT("MinifyShader: MainCS: contains FTypedefUnusedStruct"), ChunkPresent(MinifiedParsed, TEXT("FTypedefUnusedStruct")));
+		TestFalse(TEXT("MinifyShader: MainCS: contains FTypedefUnused"), ChunkPresent(MinifiedParsed, TEXT("FTypedefUnused")));
+		TestFalse(TEXT("MinifyShader: MainCS: contains TypedefUnusedBuffer"), ChunkPresent(MinifiedParsed, TEXT("TypedefUnusedBuffer")));
 	}
 
 	int32 NumErrors = ExecutionInfo.GetErrorTotal();
