ShaderMinifier - Add basic support for global operator overloads

diff --git a/Engine/Source/Developer/ShaderCompilerCommon/Private/ShaderMinifier.cpp b/Engine/Source/Developer/ShaderCompilerCommon/Private/ShaderMinifier.cpp
--- a/Engine/Source/Developer/ShaderCompilerCommon/Private/ShaderMinifier.cpp
+++ b/Engine/Source/Developer/ShaderCompilerCommon/Private/ShaderMinifier.cpp
@@ -73,6 +73,74 @@ static bool IsPossibleIdentifierCharacter(TCHAR C)
 	return (C >= '0' && C <= '9') || (C >= 'a' && C <= 'z') || (C >= 'A' && C <= 'Z') || C == '_';
 }
 
+static FStringView ExtractOperator(FStringView Source)
+{
+	FStringView Result;
+
+	if (Source.IsEmpty() || IsPossibleIdentifierCharacter(Source[0]))
+	{
+		return Result;
+	}
+
+	// NOTE: array is sorted by length to match complete operator character sequences first
+	static const FStringView SupportedOperators[] =
+	{
+		// three-character operators
+		FStringView(TEXT("<<=")),
+		FStringView(TEXT(">>=")),
+		FStringView(TEXT("->*")),
+
+		// two-character operators
+		FStringView(TEXT("+=")),
+		FStringView(TEXT("++")),
+		FStringView(TEXT("-=")),
+		FStringView(TEXT("--")),
+		FStringView(TEXT("->")),
+		FStringView(TEXT("*=")),
+		FStringView(TEXT("/=")),
+		FStringView(TEXT("%=")),
+		FStringView(TEXT("^=")),
+		FStringView(TEXT("&=")),
+		FStringView(TEXT("&&")),
+		FStringView(TEXT("|=")),
+		FStringView(TEXT("||")),
+		FStringView(TEXT("<<")),
+		FStringView(TEXT("<=")),
+		FStringView(TEXT(">>")),
+		FStringView(TEXT(">=")),
+		FStringView(TEXT("==")),
+		FStringView(TEXT("!=")),
+		FStringView(TEXT("()")),
+		FStringView(TEXT("[]")),
+
+		// single character operators
+		FStringView(TEXT("+")),
+		FStringView(TEXT("-")),
+		FStringView(TEXT("*")),
+		FStringView(TEXT("/")),
+		FStringView(TEXT("%")),
+		FStringView(TEXT("^")),
+		FStringView(TEXT("&")),
+		FStringView(TEXT("|")),
+		FStringView(TEXT("~")),
+		FStringView(TEXT("!")),
+		FStringView(TEXT("=")),
+		FStringView(TEXT("<")),
+		FStringView(TEXT(">")),
+	};
+
+	for (FStringView Operator : SupportedOperators)
+	{
+		if (Source.StartsWith(Operator))
+		{
+			Result = Source.SubStr(0, Operator.Len());
+			break;
+		}
+	}
+
+	return Result;
+}
+
 static FStringView SkipUntilNonIdentifierCharacter(FStringView Source) {
 	int32 Len = Source.Len();
 	int32 Cursor = 0;
@@ -223,6 +291,7 @@ enum class EBlockType : uint8 {
 	Directive,  // #define, #pragma, #line, etc.
 	NamespaceDelimiter, // e.g. :: in an identifier like Foo::bar
 	PtrOrRef, // e.g. '*' or '&' as part of the type
+	OperatorName, // overladed operator, e.g. '+', '+=', etc.
 };
 
 struct FCodeBlock
@@ -236,6 +305,7 @@ enum class ECodeChunkType {
 	Struct,
 	CBuffer,  // HLSL cbuffer block possibly without trailing ';'
 	Function,
+	Operator,
 	Variable,
 	Enum,
 	Define,
@@ -387,6 +457,7 @@ static FParsedShader ParseShader(FStringView InSource, FDiagnostics& Output)
 	int32 EnumBlockIndex = INDEX_NONE;
 	int32 BodyBlockIndex = INDEX_NONE;
 	int32 ExpressionBlockIndex = INDEX_NONE;
+	int32 OperatorKeywordBlockIndex = INDEX_NONE;
 
 	FNamespaceTracker NamespaceTracker;
 	FStringView		  PendingNamespace;
@@ -411,7 +482,7 @@ static FParsedShader ParseShader(FStringView InSource, FDiagnostics& Output)
 		PendingBlocks.Push(NewBlock);
 	};
 
-	auto FinalizeChunk = [&]() 
+	auto FinalizeChunk = [&]()
 	{
 		const bool bFoundArgs = ArgsBlockIndex >= 0;
 
@@ -508,6 +579,17 @@ static FParsedShader ParseShader(FStringView InSource, FDiagnostics& Output)
 				}
 				PendingBlocks[NameBlockIndex].Type = EBlockType::Name;
 			}
+			else if (ChunkType == ECodeChunkType::Operator)
+			{
+				// Treat any uncategorized blocks as part of the type
+				for (int32 Index = 0; Index < OperatorKeywordBlockIndex; Index++)
+				{
+					if (PendingBlocks[Index].Type == EBlockType::Unknown)
+					{
+						PendingBlocks[Index].Type = EBlockType::Type;
+					}
+				}
+			}
 
 			if (ChunkType == ECodeChunkType::Struct && bHasName && !bHasType)
 			{
@@ -564,6 +646,7 @@ static FParsedShader ParseShader(FStringView InSource, FDiagnostics& Output)
 			EnumBlockIndex = INDEX_NONE;
 			BodyBlockIndex = INDEX_NONE;
 			ExpressionBlockIndex = INDEX_NONE;
+			OperatorKeywordBlockIndex = INDEX_NONE;
 			bFoundBody = false;
 			bFoundColon = false;
 			bFoundIdentifier = false;
@@ -679,6 +762,25 @@ static FParsedShader ParseShader(FStringView InSource, FDiagnostics& Output)
 				AddDiagnostic(Output.Errors, TEXT("Expected token '}'"));
 			}
 		}
+		else if (ChunkType == ECodeChunkType::Operator
+			&& !PendingBlocks.IsEmpty()
+			&& OperatorKeywordBlockIndex == (PendingBlocks.Num()-1))
+		{
+			// Operator keyword found in the last processed block. Expect to find operator name character sequence next.
+
+			FStringView OperatorName = ExtractOperator(Source);
+
+			if (OperatorName.IsEmpty())
+			{
+				AddDiagnostic(Output.Errors, TEXT("Unexpected operator overload type"));
+				break;
+			}
+
+			AddBlock(EBlockType::OperatorName, OperatorName);
+			Source = SubStrView(Source, OperatorName.Len());
+
+			continue;
+		}
 
 		FStringView Remainder = SkipUntilNonIdentifierCharacter(Source);
 		FStringView Identifier = SubStrView(Source, 0, Source.Len() - Remainder.Len());
@@ -722,6 +824,20 @@ static FParsedShader ParseShader(FStringView InSource, FDiagnostics& Output)
 					AddBlock(EBlockType::Keyword, Identifier);
 					continue;
 				}
+				else if (Identifier == TEXT("template"))
+				{
+					Source = Remainder;
+					AddBlock(EBlockType::Keyword, Identifier);
+					continue;
+				}
+				else if (Identifier == TEXT("operator"))
+				{
+					ChunkType = ECodeChunkType::Operator;
+					Source = Remainder;
+					OperatorKeywordBlockIndex = PendingBlocks.Num();
+					AddBlock(EBlockType::Keyword, Identifier);
+					continue;
+				}
 			}
 			else if (ChunkType == ECodeChunkType::Namespace)
 			{
@@ -1075,6 +1191,7 @@ static void OutputChunk(const FCodeChunk& Chunk, FStringBuilderBase& OutputStrea
 	}
 
 	if (Chunk.Type != ECodeChunkType::Function
+		&& Chunk.Type != ECodeChunkType::Operator
 		&& Chunk.Type != ECodeChunkType::CBuffer
 		&& Chunk.Type != ECodeChunkType::Pragma
 		&& Chunk.Type != ECodeChunkType::Define
@@ -1286,6 +1403,8 @@ static FString MinifyShader(const FParsedShader& Parsed, TConstArrayView<FString
 		ProcessedIdentifiers.Add(TEXT("void"));
 		ProcessedIdentifiers.Add(TEXT("while"));
 		ProcessedIdentifiers.Add(TEXT("typedef"));
+		ProcessedIdentifiers.Add(TEXT("template"));
+		ProcessedIdentifiers.Add(TEXT("operator"));
 
 		// HLSL resource types
 		ProcessedIdentifiers.Add(TEXT("TextureCubeArray"));
@@ -1375,7 +1494,7 @@ static FString MinifyShader(const FParsedShader& Parsed, TConstArrayView<FString
 
 	if (PendingChunks.IsEmpty())
 	{
-		// printf("Entry point chunk is not found in the shader\n");
+		// Entry point chunk is not found in the shader
 		return {};
 	}
 
@@ -1386,6 +1505,11 @@ static FString MinifyShader(const FParsedShader& Parsed, TConstArrayView<FString
 	{
 		for (const FCodeBlock& Block : Chunk.Blocks)
 		{
+			if (Block.Type == EBlockType::Keyword)
+			{
+				continue;
+			}
+
 			if (Chunk.Type == ECodeChunkType::Function && Block.Type != EBlockType::Name)
 			{
 				continue;
@@ -1413,6 +1537,13 @@ static FString MinifyShader(const FParsedShader& Parsed, TConstArrayView<FString
 				continue;
 			}
 
+			if (Chunk.Type == ECodeChunkType::Operator
+				&& Block.Type != EBlockType::Type
+				&& Block.Type != EBlockType::Args)
+			{
+				continue;
+			}
+
 			ChunksByIdentifier.FindOrAdd(Block.Code).Push(&Chunk);
 		}
 	}
@@ -1969,6 +2100,25 @@ bool FShaderMinifierParserTest::RunTest(const FString& Parameters)
 		}
 	}
 
+	{
+		auto P = ParseShader(TEXT("template< typename T > TMyStruct<T> operator + ( TMyStruct<T> A, T B ) { /*...*/ }"));
+		if (TestEqual(TEXT("ParseShader: operators: num chunks"), P.Chunks.Num(), 1))
+		{
+			TestEqual(TEXT("ParseShader: operators: chunk type"), P.Chunks[0].Type, ECodeChunkType::Operator);
+			if (TestEqual(TEXT("ParseShader: operators: chunk 0: num blocks"), P.Chunks[0].Blocks.Num(), 8))
+			{
+				FString Args = FString(P.Chunks[0].FindFirstBlockByType(EBlockType::Args));
+				TestEqual(TEXT("ParseShader: operators: chunk 0: type name"), *Args, TEXT("( TMyStruct<T> A, T B )"));
+
+				FString TypeName = FString(P.Chunks[0].FindFirstBlockByType(EBlockType::Type));
+				TestEqual(TEXT("ParseShader: operators: chunk 0: type name"), *TypeName, TEXT("TMyStruct"));
+
+				FString OperatorName = FString(P.Chunks[0].FindFirstBlockByType(EBlockType::OperatorName));
+				TestEqual(TEXT("ParseShader: operators: chunk 0: operator name"), *OperatorName, TEXT("+"));
+			}
+		}
+	}
+
 	{
 		auto P = ParseShader(TEXT("typedef Bar Foo;"));
 		if (TestEqual(TEXT("ParseShader: standard typedef: num chunks"), P.Chunks.Num(), 1))
@@ -2105,6 +2255,12 @@ struct FTypedefUnusedStruct
 typedef StructuredBuffer<FTypedefUnusedStruct> FTypedefUnused;
 FTypedefUnused TypedefUnusedBuffer;
 
+template<typename T> struct TUsedTemplate { T Value; };
+template<typename T> TUsedTemplate<T> operator*(TUsedTemplate<T> A, T B) { return (TUsedTemplate<T>)0; }
+
+template<typename T> struct TUnusedTemplate { T Value[2]; };
+template<typename T> TUnusedTemplate<T> operator%(TUnusedTemplate<T> A, T B) { return (TUnusedTemplate<T>)0; }
+
 // Test comment 2
 [numthreads(1,1,1)]
 // Comment during function declaration
@@ -2112,6 +2268,7 @@ void MainCS()
 {
 	using namespace NS1::NS2;
 	using namespace NS3;
+	TUsedTemplate Foo;
 	float A = FunB(GAnonymousStruct.Foo);
 	float B = FunB(GStructA.Bar + GStructB.Foo + GStructC.Foo);
 	float C = FunB(GInitializedAnonymousStructA.Foo + GInitializedAnonymousStructB.Foo);
@@ -2173,6 +2330,8 @@ void MainCS()
 		TestTrue(TEXT("MinifyShader: MainCS: contains FTypedefUsed"), ChunkPresent(MinifiedParsed, TEXT("FTypedefUsed")));
 		TestTrue(TEXT("MinifyShader: MainCS: contains FTypedefUsedChained"), ChunkPresent(MinifiedParsed, TEXT("FTypedefUsedChained")));
 		TestTrue(TEXT("MinifyShader: MainCS: contains TypedefUsedBuffer"), ChunkPresent(MinifiedParsed, TEXT("TypedefUsedBuffer")));
+		TestTrue(TEXT("MinifyShader: MainCS: contains struct TUnusedTemplate"), MinifiedParsed.Source.Contains(TEXT("struct TUsedTemplate")));
+		TestTrue(TEXT("MinifyShader: MainCS: contains TUsedTemplate<T> operator*"), MinifiedParsed.Source.Contains(TEXT("TUsedTemplate<T> operator*")));
 
 		// Expect false:
 		TestFalse(TEXT("MinifyShader: MainCS: contains UnreferencedFunction"), ChunkPresent(MinifiedParsed, TEXT("UnreferencedFunction")));
@@ -2182,6 +2341,8 @@ void MainCS()
 		TestFalse(TEXT("MinifyShader: MainCS: contains FTypedefUnusedStruct"), ChunkPresent(MinifiedParsed, TEXT("FTypedefUnusedStruct")));
 		TestFalse(TEXT("MinifyShader: MainCS: contains FTypedefUnused"), ChunkPresent(MinifiedParsed, TEXT("FTypedefUnused")));
 		TestFalse(TEXT("MinifyShader: MainCS: contains TypedefUnusedBuffer"), ChunkPresent(MinifiedParsed, TEXT("TypedefUnusedBuffer")));
+		TestFalse(TEXT("MinifyShader: MainCS: contains struct TUnusedTemplate"), MinifiedParsed.Source.Contains(TEXT("struct TUnusedTemplate")));
+		TestFalse(TEXT("MinifyShader: MainCS: contains TUnusedTemplate<T> operator%"), MinifiedParsed.Source.Contains(TEXT("TUnusedTemplate<T> operator%")));
 	}
 
 	int32 NumErrors = ExecutionInfo.GetErrorTotal();
