ShaderMinifier - Don't skip keyword blocks when building ChunksByIdentifier (temporary workaround for an issue introduced in 28384399)

diff --git a/Engine/Source/Developer/ShaderCompilerCommon/Private/ShaderMinifier.cpp b/Engine/Source/Developer/ShaderCompilerCommon/Private/ShaderMinifier.cpp
--- a/Engine/Source/Developer/ShaderCompilerCommon/Private/ShaderMinifier.cpp
+++ b/Engine/Source/Developer/ShaderCompilerCommon/Private/ShaderMinifier.cpp
@@ -1505,11 +1505,6 @@ static FString MinifyShader(const FParsedShader& Parsed, TConstArrayView<FString
 	{
 		for (const FCodeBlock& Block : Chunk.Blocks)
 		{
-			if (Block.Type == EBlockType::Keyword)
-			{
-				continue;
-			}
-
 			if (Chunk.Type == ECodeChunkType::Function && Block.Type != EBlockType::Name)
 			{
 				continue;
