Fixing enum handling in the shader minifier. Since the enum elements weren't marked as reachable identifiers, the minifier was considering the whole enum block unused.

diff --git a/Engine/Source/Developer/ShaderCompilerCommon/Private/ShaderMinifier.cpp b/Engine/Source/Developer/ShaderCompilerCommon/Private/ShaderMinifier.cpp
--- a/Engine/Source/Developer/ShaderCompilerCommon/Private/ShaderMinifier.cpp
+++ b/Engine/Source/Developer/ShaderCompilerCommon/Private/ShaderMinifier.cpp
@@ -1405,6 +1405,7 @@ static FString MinifyShader(const FParsedShader& Parsed, TConstArrayView<FString
 		ProcessedIdentifiers.Add(TEXT("typedef"));
 		ProcessedIdentifiers.Add(TEXT("template"));
 		ProcessedIdentifiers.Add(TEXT("operator"));
+		ProcessedIdentifiers.Add(TEXT("enum"));
 
 		// HLSL resource types
 		ProcessedIdentifiers.Add(TEXT("TextureCubeArray"));
@@ -1505,6 +1506,11 @@ static FString MinifyShader(const FParsedShader& Parsed, TConstArrayView<FString
 	{
 		for (const FCodeBlock& Block : Chunk.Blocks)
 		{
+			if (Block.Type == EBlockType::Keyword)
+			{
+				continue;
+			}
+
 			if (Chunk.Type == ECodeChunkType::Function && Block.Type != EBlockType::Name)
 			{
 				continue;
@@ -1520,7 +1526,7 @@ static FString MinifyShader(const FParsedShader& Parsed, TConstArrayView<FString
 				continue;
 			}
 
-			if (Chunk.Type == ECodeChunkType::CBuffer && Block.Type == EBlockType::Body)
+			if ((Chunk.Type == ECodeChunkType::CBuffer || Chunk.Type == ECodeChunkType::Enum) && Block.Type == EBlockType::Body)
 			{
 				TempIdentifiers.Empty();
 				ExtractIdentifiers(Block.Code, TempIdentifiers);
@@ -2256,6 +2262,18 @@ template<typename T> TUsedTemplate<T> operator*(TUsedTemplate<T> A, T B) { retur
 template<typename T> struct TUnusedTemplate { T Value[2]; };
 template<typename T> TUnusedTemplate<T> operator%(TUnusedTemplate<T> A, T B) { return (TUnusedTemplate<T>)0; }
 
+enum EEnumUsed : int
+{
+	ENUM_USED_PART_1 = 0,
+	ENUM_USED_PART_2 = 1,
+};
+
+enum EEnumUnused
+{
+	ENUM_UNUSED_PART_1,
+	ENUM_UNUSED_PART_2,
+};
+
 // Test comment 2
 [numthreads(1,1,1)]
 // Comment during function declaration
@@ -2267,7 +2285,7 @@ void MainCS()
 	float A = FunB(GAnonymousStruct.Foo);
 	float B = FunB(GStructA.Bar + GStructB.Foo + GStructC.Foo);
 	float C = FunB(GInitializedAnonymousStructA.Foo + GInitializedAnonymousStructB.Foo);
-	float D = TypedefUsedBuffer[0].Foo;
+	float D = TypedefUsedBuffer[ENUM_USED_PART_1].Foo;
 	OutputBuffer[0] = A + B + D;
 }
 )");
@@ -2327,6 +2345,9 @@ void MainCS()
 		TestTrue(TEXT("MinifyShader: MainCS: contains TypedefUsedBuffer"), ChunkPresent(MinifiedParsed, TEXT("TypedefUsedBuffer")));
 		TestTrue(TEXT("MinifyShader: MainCS: contains struct TUnusedTemplate"), MinifiedParsed.Source.Contains(TEXT("struct TUsedTemplate")));
 		TestTrue(TEXT("MinifyShader: MainCS: contains TUsedTemplate<T> operator*"), MinifiedParsed.Source.Contains(TEXT("TUsedTemplate<T> operator*")));
+		TestTrue(TEXT("MinifyShader: MainCS: contains EEnumUsed"), MinifiedParsed.Source.Contains(TEXT("EEnumUsed")));
+		TestTrue(TEXT("MinifyShader: MainCS: contains ENUM_USED_PART_1"), MinifiedParsed.Source.Contains(TEXT("ENUM_USED_PART_1")));
+		TestTrue(TEXT("MinifyShader: MainCS: contains ENUM_USED_PART_2"), MinifiedParsed.Source.Contains(TEXT("ENUM_USED_PART_2")));
 
 		// Expect false:
 		TestFalse(TEXT("MinifyShader: MainCS: contains UnreferencedFunction"), ChunkPresent(MinifiedParsed, TEXT("UnreferencedFunction")));
@@ -2338,6 +2359,9 @@ void MainCS()
 		TestFalse(TEXT("MinifyShader: MainCS: contains TypedefUnusedBuffer"), ChunkPresent(MinifiedParsed, TEXT("TypedefUnusedBuffer")));
 		TestFalse(TEXT("MinifyShader: MainCS: contains struct TUnusedTemplate"), MinifiedParsed.Source.Contains(TEXT("struct TUnusedTemplate")));
 		TestFalse(TEXT("MinifyShader: MainCS: contains TUnusedTemplate<T> operator%"), MinifiedParsed.Source.Contains(TEXT("TUnusedTemplate<T> operator%")));
+		TestFalse(TEXT("MinifyShader: MainCS: contains EEnumUnused"), MinifiedParsed.Source.Contains(TEXT("EEnumUnused")));
+		TestFalse(TEXT("MinifyShader: MainCS: contains ENUM_UNUSED_PART_1"), MinifiedParsed.Source.Contains(TEXT("ENUM_UNUSED_PART_1")));
+		TestFalse(TEXT("MinifyShader: MainCS: contains ENUM_UNUSED_PART_2"), MinifiedParsed.Source.Contains(TEXT("ENUM_UNUSED_PART_2")));
 	}
 
 	int32 NumErrors = ExecutionInfo.GetErrorTotal();
