Shader Compiler:  CleanupUniformBufferCodeNew minor fix to account for escaped quotes in text.  Continue string parsing if character before quote is a backslash.  Note that there are no escaped quotes in any existing .usf or .ush file at the moment, so I don't think the bug should have caused any issues.

#rnx

diff --git a/Engine/Source/Developer/ShaderCompilerCommon/Private/ShaderCompilerCommon.cpp b/Engine/Source/Developer/ShaderCompilerCommon/Private/ShaderCompilerCommon.cpp
--- a/Engine/Source/Developer/ShaderCompilerCommon/Private/ShaderCompilerCommon.cpp
+++ b/Engine/Source/Developer/ShaderCompilerCommon/Private/ShaderCompilerCommon.cpp
@@ -1262,9 +1262,9 @@ bool FindNextCompoundIdentifier(const TCHAR*& Search, const uint64 RootIdentifie
 		}
 		else if (SearchCharFlag & (uint8)AsciiFlags::Quote)
 		{
-			// Quote, skip to next Quote (or maybe end of string if text is malformed)
+			// Quote, skip to next Quote (or maybe end of string if text is malformed), ignoring the quote if it's escaped
 			SearchChar++;
-			while (*SearchChar && *SearchChar != TEXT('\"'))
+			while (*SearchChar && (*SearchChar != TEXT('\"') || *(SearchChar - 1) == TEXT('\\')))
 			{
 				SearchChar++;
 			}
