ShaderMinifier - Use inline implementations of string view Equals and StartsWith

* Also use TEXTVIEW instead of TEXT where possible
* ~1.5x speedup

diff --git a/Engine/Source/Developer/ShaderCompilerCommon/Private/ShaderMinifier.cpp b/Engine/Source/Developer/ShaderCompilerCommon/Private/ShaderMinifier.cpp
--- a/Engine/Source/Developer/ShaderCompilerCommon/Private/ShaderMinifier.cpp
+++ b/Engine/Source/Developer/ShaderCompilerCommon/Private/ShaderMinifier.cpp
@@ -35,7 +35,7 @@ template<typename TCondition>
 static FStringView SkipUntil(FStringView Source, TCondition Cond)
 {
 	int32 Cursor = 0;
-	int32 SourceLen = Source.Len();
+	const int32 SourceLen = Source.Len();
 	while (Cursor < SourceLen)
 	{
 		if (Cond(FStringView(Source.GetData() + Cursor, SourceLen - Cursor)))
@@ -47,6 +47,37 @@ static FStringView SkipUntil(FStringView Source, TCondition Cond)
 	return FStringView(Source.GetData() + Cursor, SourceLen - Cursor);
 }
 
+static bool Equals(FStringView A, FStringView B)
+{
+	int32 Len = A.Len();
+	if (Len != B.Len())
+	{
+		return false;
+	}
+	const TCHAR* DataA = A.GetData();
+	const TCHAR* DataB = B.GetData();
+	for (int32 I = 0; I < Len; ++I)
+	{
+		if (DataA[I] != DataB[I])
+		{
+			return false;
+		}
+	}
+	return true;
+}
+
+static bool StartsWith(FStringView Source, FStringView Prefix)
+{
+	const int32 SourceLen = Source.Len();
+	const int32 PrefixLen = Prefix.Len();
+	if (PrefixLen > SourceLen)
+	{
+		return false;
+	}
+	FStringView SourceView(Source.GetData(), PrefixLen);
+	return Equals(SourceView, Prefix);
+}
+
 static bool IsSpace(TCHAR C)
 {
 	switch (C)
@@ -131,7 +162,7 @@ static FStringView ExtractOperator(FStringView Source)
 
 	for (FStringView Operator : SupportedOperators)
 	{
-		if (Source.StartsWith(Operator))
+		if (StartsWith(Source, Operator))
 		{
 			Result = Source.SubStr(0, Operator.Len());
 			break;
@@ -141,11 +172,12 @@ static FStringView ExtractOperator(FStringView Source)
 	return Result;
 }
 
-static FStringView SkipUntilNonIdentifierCharacter(FStringView Source) {
-	int32 Len = Source.Len();
+static FStringView SkipUntilNonIdentifierCharacter(FStringView Source) 
+{
+	const int32 SourceLen = Source.Len();
 	int32 Cursor = 0;
 	const TCHAR* SourceData = Source.GetData();
-	while (Cursor < Len)
+	while (Cursor < SourceLen)
 	{
 		if (!IsPossibleIdentifierCharacter(SourceData[Cursor]))
 		{
@@ -153,14 +185,15 @@ static FStringView SkipUntilNonIdentifierCharacter(FStringView Source) {
 		}
 		++Cursor;
 	}
-	return FStringView(SourceData + Cursor, Len - Cursor);
+	return FStringView(SourceData + Cursor, SourceLen - Cursor);
 }
 
-static FStringView SkipUntilNonNumber(FStringView Source) {
-	int32 Len = Source.Len();
+static FStringView SkipUntilNonNumber(FStringView Source) 
+{
+	const int32 SourceLen = Source.Len();
 	int32 Cursor = 0;
 	const TCHAR* SourceData = Source.GetData();
-	while (Cursor < Len)
+	while (Cursor < SourceLen)
 	{
 		if (!IsNumber(SourceData[Cursor]))
 		{
@@ -168,15 +201,15 @@ static FStringView SkipUntilNonNumber(FStringView Source) {
 		}
 		++Cursor;
 	}
-	return FStringView(SourceData + Cursor, Len - Cursor);
+	return FStringView(SourceData + Cursor, SourceLen - Cursor);
 }
 
 static FStringView SkipSpace(FStringView Source)
 {
-	int32 Len = Source.Len();
+	const int32 SourceLen = Source.Len();
 	int32 Cursor = 0;
 	const TCHAR* SourceData = Source.GetData();
-	while (Cursor < Len)
+	while (Cursor < SourceLen)
 	{
 		if (!IsSpace(SourceData[Cursor]))
 		{
@@ -184,7 +217,7 @@ static FStringView SkipSpace(FStringView Source)
 		}
 		++Cursor;
 	}
-	return FStringView(SourceData + Cursor, Len - Cursor);
+	return FStringView(SourceData + Cursor, SourceLen - Cursor);
 }
 
 static FStringView TrimSpace(FStringView Source)
@@ -230,7 +263,7 @@ static FStringView SkipUntilNextLine(FStringView Source)
 
 static FStringView SkipUntilStr(FStringView Haystack, FStringView Needle)
 {
-	return SkipUntil(Haystack, [Needle](FStringView  S) { return S.StartsWith(Needle, ESearchCase::CaseSensitive); });
+	return SkipUntil(Haystack, [Needle](FStringView  S) { return StartsWith(S, Needle); });
 }
 
 static FStringView ExtractBlock(FStringView Source, TCHAR DelimBegin, TCHAR DelimEnd)
@@ -240,9 +273,11 @@ static FStringView ExtractBlock(FStringView Source, TCHAR DelimBegin, TCHAR Deli
 
 	int32 PosEnd = INDEX_NONE;
 	int32 Stack  = 0;
-	for (int32 I = 0; I < Source.Len(); ++I)
+	const int32 SourceLen = Source.Len();
+	const TCHAR* SourceData = Source.GetData();
+	for (int32 I = 0; I < SourceLen; ++I)
 	{
-		TCHAR C = Source[I];
+		TCHAR C = SourceData[I];
 		if (C == DelimBegin)
 		{
 			Stack++;
@@ -665,7 +700,7 @@ static FParsedShader ParseShader(FStringView InSource, FDiagnostics& Output)
 			break;
 		}
 
-		if (Source.StartsWith(TEXT("//")))
+		if (StartsWith(Source, TEXTVIEW("//")))
 		{
 			FStringView Remainder = SkipUntilNextLine(Source);
 
@@ -682,12 +717,12 @@ static FParsedShader ParseShader(FStringView InSource, FDiagnostics& Output)
 
 			continue;
 		}
-		else if (Source.StartsWith(TEXT("#line")))
+		else if (StartsWith(Source, TEXTVIEW("#line")))
 		{
 			Source = SkipUntilNextLine(Source);
 			continue;
 		}
-		else if (Source.StartsWith(TEXT("#pragma")))
+		else if (StartsWith(Source, TEXTVIEW("#pragma")))
 		{
 			FStringView Remainder = SkipUntilNextLine(Source);
 			FStringView Block = SubStrView(Source, 0, Source.Len() - Remainder.Len());
@@ -697,7 +732,7 @@ static FParsedShader ParseShader(FStringView InSource, FDiagnostics& Output)
 			Source = Remainder;
 			continue;
 		}
-		else if (Source.StartsWith(TEXT("#define")))
+		else if (StartsWith(Source, TEXTVIEW("#define")))
 		{
 			// TODO: handle `\` new lines in defines
 			FStringView Remainder = SkipUntilNextLine(Source);
@@ -708,25 +743,25 @@ static FParsedShader ParseShader(FStringView InSource, FDiagnostics& Output)
 			Source = Remainder;
 			continue;
 		}
-		else if (Source.StartsWith(TEXT("#if 0")))
+		else if (StartsWith(Source, TEXTVIEW("#if 0")))
 		{
-			Source = SkipUntilStr(Source, TEXT("#endif"));
+			Source = SkipUntilStr(Source, TEXTVIEW("#endif"));
 			if (Source.Len() >= 6)
 			{
 				Source = SubStrView(Source, 6);
 			}
 			continue;
 		}
-		else if (Source.StartsWith(TEXT("/*")))
+		else if (StartsWith(Source, TEXTVIEW("/*")))
 		{
-			Source = SkipUntilStr(Source, TEXT("*/"));
+			Source = SkipUntilStr(Source, TEXTVIEW("*/"));
 			if (Source.Len() >= 2)
 			{
 				Source = SubStrView(Source, 2);
 			}
 			continue;
 		}
-		else if (PendingBlocks.IsEmpty() && Source.StartsWith(TEXT("{")))
+		else if (PendingBlocks.IsEmpty() && StartsWith(Source, TEXTVIEW("{")))
 		{
 			if (ChunkType == ECodeChunkType::Namespace)
 			{
@@ -750,7 +785,7 @@ static FParsedShader ParseShader(FStringView InSource, FDiagnostics& Output)
 			}
 			continue;
 		}
-		else if (PendingBlocks.IsEmpty() && Source.StartsWith(TEXT("}")))
+		else if (PendingBlocks.IsEmpty() && StartsWith(Source, TEXTVIEW("}")))
 		{
 			if (NamespaceTracker.Pop())
 			{
@@ -789,48 +824,48 @@ static FParsedShader ParseShader(FStringView InSource, FDiagnostics& Output)
 		{
 			if (ChunkType == ECodeChunkType::Unknown)
 			{
-				if (Identifier == TEXT("struct"))
+				if (Equals(Identifier, TEXTVIEW("struct")))
 				{
 					ChunkType = ECodeChunkType::Struct;
 					StructBlockIndex = PendingBlocks.Num();
 				}
-				else if (Identifier == TEXT("cbuffer") || Identifier == TEXT("ConstantBuffer"))
+				else if (Equals(Identifier, TEXTVIEW("cbuffer")) || Equals(Identifier, TEXTVIEW("ConstantBuffer")))
 				{
 					ChunkType = ECodeChunkType::CBuffer;
 					CbufferBlockIndex = PendingBlocks.Num();
 				}
-				else if (Identifier == TEXT("enum"))
+				else if (Equals(Identifier, TEXTVIEW("enum")))
 				{
 					ChunkType = ECodeChunkType::Enum;
 					EnumBlockIndex = PendingBlocks.Num();
 				}
-				else if (Identifier == TEXT("namespace"))
+				else if (Equals(Identifier, TEXTVIEW("namespace")))
 				{
 					ChunkType = ECodeChunkType::Namespace;
 					Source = Remainder;
 					continue;
 				}
-				else if (Identifier == TEXT("using"))
+				else if (Equals(Identifier, TEXTVIEW("using")))
 				{
 					ChunkType = ECodeChunkType::Using;
 					Source = Remainder;
 					AddBlock(EBlockType::Keyword, Identifier);
 					continue;
 				}
-				else if (Identifier == TEXT("typedef"))
+				else if (Equals(Identifier, TEXTVIEW("typedef")))
 				{
 					ChunkType = ECodeChunkType::Typedef;
 					Source = Remainder;
 					AddBlock(EBlockType::Keyword, Identifier);
 					continue;
 				}
-				else if (Identifier == TEXT("template"))
+				else if (Equals(Identifier, TEXTVIEW("template")))
 				{
 					Source = Remainder;
 					AddBlock(EBlockType::Keyword, Identifier);
 					continue;
 				}
-				else if (Identifier == TEXT("operator"))
+				else if (Equals(Identifier, TEXTVIEW("operator")))
 				{
 					ChunkType = ECodeChunkType::Operator;
 					Source = Remainder;
@@ -876,12 +911,12 @@ static FParsedShader ParseShader(FStringView InSource, FDiagnostics& Output)
 
 		EBlockType BlockType = EBlockType::Unknown;
 
-		if (Source.StartsWith(TEXT("==")))
+		if (StartsWith(Source, TEXTVIEW("==")))
 		{
 			AddDiagnostic(Output.Errors, TEXT("Unexpected sequence '=='"));
 			break;
 		}
-		else if (Source.StartsWith(TEXT("::")))
+		else if (StartsWith(Source, TEXTVIEW("::")))
 		{
 			Block = SubStrView(Source, 0, 2);
 			Source = SubStrView(Source, 2);
@@ -1052,8 +1087,9 @@ static TArray<FStringView> SplitByChar(FStringView Source, TCHAR Delimiter)
 	TArray<FStringView> Result;
 
 	int32 Start = 0;
+	const int32 SourceLen = Source.Len();
 
-	for (int32 I = 0; I < Source.Len(); ++I)
+	for (int32 I = 0; I < SourceLen; ++I)
 	{
 		TCHAR C = Source[I];
 		if (C == Delimiter)
@@ -1077,8 +1113,6 @@ static void ExtractIdentifiers(FStringView InSource, TArray<FStringView>& Result
 {
 	FStringView Source = InSource;
 
-	const ESearchCase::Type SC = ESearchCase::CaseSensitive;
-
 	for (;;)
 	{
 		Source = SkipSpace(Source);
@@ -1090,25 +1124,25 @@ static void ExtractIdentifiers(FStringView InSource, TArray<FStringView>& Result
 
 		if (!IsPossibleIdentifierCharacter(Source[0]))
 		{
-			if (Source.StartsWith(TEXT("//"), SC) 
-				|| Source.StartsWith(TEXT("#line"), SC) 
-				|| Source.StartsWith(TEXT("#pragma"), SC))
+			if (StartsWith(Source, TEXTVIEW("//"))
+				|| StartsWith(Source, TEXTVIEW("#line"))
+				|| StartsWith(Source, TEXTVIEW("#pragma")))
 			{
 				Source = SkipUntilNextLine(Source);
 				continue;
 			}
-			else if (Source.StartsWith(TEXT("#if 0"), SC))
+			else if (StartsWith(Source, TEXTVIEW("#if 0")))
 			{
-				Source = SkipUntilStr(Source, TEXT("#endif"));
+				Source = SkipUntilStr(Source, TEXTVIEW("#endif"));
 				if (Source.Len() >= 6)
 				{
 					Source = SubStrView(Source, 6);
 				}
 				continue;
 			}
-			else if (Source.StartsWith(TEXT("/*"), SC))
+			else if (StartsWith(Source, TEXTVIEW("/*")))
 			{
-				Source = SkipUntilStr(Source, TEXT("*/"));
+				Source = SkipUntilStr(Source, TEXTVIEW("*/"));
 				if (Source.Len() >= 2)
 				{
 					Source = SubStrView(Source, 2);
@@ -1208,7 +1242,7 @@ struct FCasedStringViewKeyFuncs : public DefaultKeyFuncs<FStringView>
 	static FORCEINLINE FStringView GetSetKey(FStringView K) { return K; }
 	template <typename T>
 	static FORCEINLINE FStringView GetSetKey(const TPair<FStringView, T>& P) { return P.Key; }
-	static FORCEINLINE bool Matches(FStringView A, FStringView B) { return A.Equals(B, ESearchCase::CaseSensitive); }
+	static FORCEINLINE bool Matches(FStringView A, FStringView B) { return Equals(A, B); }
 	static FORCEINLINE uint32 GetKeyHash(FStringView Key)
 	{
 		return CityHash32((const char*)Key.GetData(), Key.Len() * sizeof(*Key.GetData()));
@@ -1222,10 +1256,10 @@ static void BuildLineBreakMap(FStringView Source, TArray<int32>& OutLineBreakMap
 
 	OutLineBreakMap.Add(0); // Lines numbers are 1-based, so add a dummy element to make LowerBound later return the line number directly
 
-	const int32 Len = Source.Len();
+	const int32 SourceLen = Source.Len();
 	const TCHAR* Chars = Source.GetData(); // avoid bounds check overhead in [] operator
 
-	for (int32 Index = 0; Index < Len; ++Index)
+	for (int32 Index = 0; Index < SourceLen; ++Index)
 	{
 		if (Chars[Index] == TCHAR('\n'))
 		{
@@ -1237,7 +1271,7 @@ static void BuildLineBreakMap(FStringView Source, TArray<int32>& OutLineBreakMap
 			// However we expect input source to be fully preprocessed and comments to be removed.
 
 			FStringView PossibleDirective = Source.Mid(Index);
-			if (PossibleDirective.StartsWith(TEXT("#line")))
+			if (StartsWith(PossibleDirective, TEXTVIEW("#line")))
 			{
 				FStringView Remainder = SkipUntilNextLine(PossibleDirective);
 				FStringView LineDirective = SubStrView(PossibleDirective, 0, PossibleDirective.Len() - Remainder.Len());
@@ -1280,7 +1314,7 @@ static int32 FindLineNumber(FStringView Source, const TArray<int32>& LineBreakMa
 
 static bool ParseLineDirective(FStringView Input, int32& OutLineNumber, FStringView& OutFileName)
 {
-	if (!Input.StartsWith(TEXT("#line")))
+	if (!StartsWith(Input, TEXTVIEW("#line")))
 	{
 		return false;
 	}
