HlslParser
- Adding ability to 'forward declare' symbols per platform.
- Removes hardcoded symbols in HlslParser.cpp
- Converting a few array/string input types to views.

diff --git a/Engine/Source/Developer/ShaderCompilerCommon/Private/HlslExpressionParser.inl b/Engine/Source/Developer/ShaderCompilerCommon/Private/HlslExpressionParser.inl
--- a/Engine/Source/Developer/ShaderCompilerCommon/Private/HlslExpressionParser.inl
+++ b/Engine/Source/Developer/ShaderCompilerCommon/Private/HlslExpressionParser.inl
@@ -42,7 +42,17 @@ namespace CrossCompiler
 
 		void Add(const FString& Type)
 		{
-			Symbols.Add(Type);
+			Symbols.Emplace(Type);
+		}
+
+		void Add(FStringView Type)
+		{
+			Symbols.Emplace(Type);
+		}
+
+		void Add(const TCHAR* Type)
+		{
+			Symbols.Emplace(Type);
 		}
 
 		static bool FindType(const FSymbolScope* Scope, const FString& Type, bool bSearchUpwards = true)
diff --git a/Engine/Source/Developer/ShaderCompilerCommon/Private/HlslParser.cpp b/Engine/Source/Developer/ShaderCompilerCommon/Private/HlslParser.cpp
--- a/Engine/Source/Developer/ShaderCompilerCommon/Private/HlslParser.cpp
+++ b/Engine/Source/Developer/ShaderCompilerCommon/Private/HlslParser.cpp
@@ -6,6 +6,7 @@
 
 #include "HlslParser.h"
 #include "HlslExpressionParser.inl"
+#include "ShaderCompilerCommon.h"
 
 namespace CrossCompiler
 {
@@ -28,7 +29,7 @@ namespace CrossCompiler
 	class FHlslParser
 	{
 	public:
-		FHlslParser(FLinearAllocator* InAllocator, FCompilerMessages& InCompilerMessages);
+		FHlslParser(FLinearAllocator* InAllocator, FCompilerMessages& InCompilerMessages, TConstArrayView<FScopedDeclarations> InScopedDeclarations);
 		FHlslScanner Scanner;
 		FCompilerMessages& CompilerMessages;
 		FSymbolScope GlobalScope;
@@ -2025,7 +2026,7 @@ Done:
 		} GStaticInitializer;
 	}
 
-	FHlslParser::FHlslParser(FLinearAllocator* InAllocator, FCompilerMessages& InCompilerMessages) :
+	FHlslParser::FHlslParser(FLinearAllocator* InAllocator, FCompilerMessages& InCompilerMessages, TConstArrayView<FScopedDeclarations> InScopedDeclarations) :
 		Scanner(InCompilerMessages),
 		CompilerMessages(InCompilerMessages),
 		GlobalScope(InAllocator, nullptr),
@@ -2034,22 +2035,21 @@ Done:
 	{
 		CurrentScope = &GlobalScope;
 
+		for (const FScopedDeclarations& ScopedDeclarations : InScopedDeclarations)
 		{
-			FCreateSymbolScope SceScope(Allocator, &CurrentScope);
-			CurrentScope->Name = TEXT("sce");
-			{
-				FCreateSymbolScope GnmScope(Allocator, &CurrentScope);
-				CurrentScope->Name = TEXT("Gnm");
+			TArray<FCreateSymbolScope> ScopeStack;
 
-				CurrentScope->Add(TEXT("Sampler"));	// sce::Gnm::Sampler
+			FSymbolScope* IterScope = CurrentScope;
+			for (FStringView Scope : ScopedDeclarations.Scope)
+			{
+				ScopeStack.Emplace(Allocator, &IterScope);
+				IterScope->Name = Scope.GetData();
+			}
 
-				CurrentScope->Add(TEXT("kAnisotropyRatio1"));	// sce::Gnm::kAnisotropyRatio1
-				CurrentScope->Add(TEXT("kBorderColorTransBlack"));	// sce::Gnm::kBorderColorTransBlack
-				CurrentScope->Add(TEXT("kDepthCompareNever"));	// sce::Gnm::kDepthCompareNever
+			for (FStringView Symbol : ScopedDeclarations.Symbols)
+			{
+				IterScope->Add(Symbol);
 			}
-			//auto* Found = CurrentScope->FindGlobalNamespace(TEXT("sce"), CurrentScope);
-			//Found = Found->FindNamespace(TEXT("Gnm"));
-			//Found->FindType(Found, TEXT("Sampler"), false);
 		}
 
 		// Register built-in structure for DXR in SM6
@@ -2058,10 +2058,10 @@ Done:
 
 	namespace Parser
 	{
-		bool Parse(const FString& Input, const FString& Filename, FCompilerMessages& OutCompilerMessages, TCallback* Callback, void* CallbackData)
+		bool Parse(const FString& Input, const FString& Filename, FCompilerMessages& OutCompilerMessages, TConstArrayView<FScopedDeclarations> InScopedDeclarations, TCallback* Callback, void* CallbackData)
 		{
 			FLinearAllocator Allocator;
-			FHlslParser Parser(&Allocator, OutCompilerMessages);
+			FHlslParser Parser(&Allocator, OutCompilerMessages, InScopedDeclarations);
 			if (!Parser.Scanner.Lex(Input, Filename))
 			{
 				return false;
@@ -2100,10 +2100,10 @@ Done:
 			return bSuccess;
 		}
 
-		bool Parse(const FString& Input, const FString& Filename, FCompilerMessages& OutCompilerMessages, TFunction< void(CrossCompiler::FLinearAllocator* Allocator, CrossCompiler::TLinearArray<CrossCompiler::AST::FNode*>& ASTNodes)> Function)
+		bool Parse(const FString& Input, const FString& Filename, FCompilerMessages& OutCompilerMessages, TConstArrayView<FScopedDeclarations> InScopedDeclarations, TFunction< void(CrossCompiler::FLinearAllocator* Allocator, CrossCompiler::TLinearArray<CrossCompiler::AST::FNode*>& ASTNodes)> Function)
 		{
 			FLinearAllocator Allocator;
-			FHlslParser Parser(&Allocator, OutCompilerMessages);
+			FHlslParser Parser(&Allocator, OutCompilerMessages, InScopedDeclarations);
 			if (!Parser.Scanner.Lex(Input, Filename))
 			{
 				return false;
diff --git a/Engine/Source/Developer/ShaderCompilerCommon/Private/HlslParser.h b/Engine/Source/Developer/ShaderCompilerCommon/Private/HlslParser.h
--- a/Engine/Source/Developer/ShaderCompilerCommon/Private/HlslParser.h
+++ b/Engine/Source/Developer/ShaderCompilerCommon/Private/HlslParser.h
@@ -10,6 +10,7 @@
 #include "HlslUtils.h"
 
 class Error;
+struct FScopedDeclarations;
 
 namespace CrossCompiler
 {
@@ -31,10 +32,10 @@ namespace CrossCompiler
 		typedef void TCallback(void* CallbackData, CrossCompiler::FLinearAllocator* Allocator, CrossCompiler::TLinearArray<CrossCompiler::AST::FNode*>& ASTNodes);
 
 		// Returns true if successfully parsed
-		SHADERCOMPILERCOMMON_API bool Parse(const FString& Input, const FString& Filename, FCompilerMessages& OutCompilerMessages, TCallback* Callback, void* CallbackData = nullptr);
+		SHADERCOMPILERCOMMON_API bool Parse(const FString& Input, const FString& Filename, FCompilerMessages& OutCompilerMessages, TConstArrayView<FScopedDeclarations> InScopedDeclarations, TCallback* Callback, void* CallbackData = nullptr);
 
 		// Returns true if successfully parsed
-		SHADERCOMPILERCOMMON_API bool Parse(const FString& Input, const FString& Filename, FCompilerMessages& OutCompilerMessages, TFunction< void(CrossCompiler::FLinearAllocator* Allocator, CrossCompiler::TLinearArray<CrossCompiler::AST::FNode*>& ASTNodes)> Function);
+		SHADERCOMPILERCOMMON_API bool Parse(const FString& Input, const FString& Filename, FCompilerMessages& OutCompilerMessages, TConstArrayView<FScopedDeclarations> InScopedDeclarations, TFunction< void(CrossCompiler::FLinearAllocator* Allocator, CrossCompiler::TLinearArray<CrossCompiler::AST::FNode*>& ASTNodes)> Function);
 
 		// Sample callback to write out all nodes into a string; pass a valid pointer FString* as OutFStringPointer
 		SHADERCOMPILERCOMMON_API void WriteNodesToString(void* OutFStringPointer, CrossCompiler::FLinearAllocator* Allocator, CrossCompiler::TLinearArray<CrossCompiler::AST::FNode*>& ASTNodes);
diff --git a/Engine/Source/Developer/ShaderCompilerCommon/Private/HlslUtils.cpp b/Engine/Source/Developer/ShaderCompilerCommon/Private/HlslUtils.cpp
--- a/Engine/Source/Developer/ShaderCompilerCommon/Private/HlslUtils.cpp
+++ b/Engine/Source/Developer/ShaderCompilerCommon/Private/HlslUtils.cpp
@@ -257,11 +257,11 @@ struct FRemoveAlgorithm
 	}
 
 	// Case-insensitive when working with Semantics
-	static bool IsStringInArray(const TArray<FString>& Array, const TCHAR* Semantic)
+	static bool IsStringInArray(const TConstArrayView<FStringView> Array, const TCHAR* Semantic)
 	{
-		for (const FString& String : Array)
+		for (FStringView String : Array)
 		{
-			if (FCString::Stricmp(*String, Semantic) == 0)
+			if (String.Equals(Semantic, ESearchCase::IgnoreCase))
 			{
 				return true;
 			}
@@ -270,11 +270,11 @@ struct FRemoveAlgorithm
 		return false;
 	};
 
-	static bool IsSubstringInArray(const TArray<FString>& Array, const TCHAR* Semantic)
+	static bool IsSubstringInArray(const TConstArrayView<FStringView> Array, const TCHAR* Semantic)
 	{
-		for (const FString& String : Array)
+		for (FStringView String : Array)
 		{
-			if (FCString::Stristr(*String, Semantic) != NULL)
+			if (UE::String::FindFirst(String, Semantic, ESearchCase::IgnoreCase) != INDEX_NONE)
 			{
 				return true;
 			}
@@ -420,8 +420,8 @@ struct FRemoveAlgorithm
 
 struct FRemoveUnusedOutputs : FRemoveAlgorithm
 {
-	const TArray<FString>& UsedOutputs;
-	const TArray<FString>& Exceptions;
+	const TConstArrayView<FStringView> UsedOutputs;
+	const TConstArrayView<FStringView> Exceptions;
 
 	struct FOutputsBodyContext : FBodyContext
 	{
@@ -445,7 +445,7 @@ struct FRemoveUnusedOutputs : FRemoveAlgorithm
 		}
 	};
 
-	FRemoveUnusedOutputs(const TArray<FString>& InUsedOutputs, const TArray<FString>& InExceptions) :
+	FRemoveUnusedOutputs(const TConstArrayView<FStringView> InUsedOutputs, const TConstArrayView<FStringView> InExceptions) :
 		UsedOutputs(InUsedOutputs),
 		Exceptions(InExceptions)
 	{
@@ -899,7 +899,12 @@ struct FRemoveUnusedOutputs : FRemoveAlgorithm
 	}
 };
 
-bool RemoveUnusedOutputs(FString& InOutSourceCode, const TArray<FString>& InUsedOutputs, const TArray<FString>& InExceptions, FString& EntryPoint, TArray<FString>& OutErrors)
+bool RemoveUnusedOutputs(FString& InOutSourceCode,
+	TConstArrayView<FStringView> InUsedOutputs,
+	TConstArrayView<FStringView> InExceptions,
+	TConstArrayView<FScopedDeclarations> InScopedDeclarations,
+	FString& EntryPoint,
+	TArray<FString>& OutErrors)
 {
 	FString DummyFilename(TEXT("/Engine/Private/RemoveUnusedOutputs.usf"));
 	FRemoveUnusedOutputs Data(InUsedOutputs, InExceptions);
@@ -915,7 +920,7 @@ bool RemoveUnusedOutputs(FString& InOutSourceCode, const TArray<FString>& InUsed
 		}
 	};
 	CrossCompiler::FCompilerMessages Messages;
-	if (!CrossCompiler::Parser::Parse(InOutSourceCode, DummyFilename, Messages, Lambda))
+	if (!CrossCompiler::Parser::Parse(InOutSourceCode, DummyFilename, Messages, InScopedDeclarations, Lambda))
 	{
 		Data.Errors.Add(FString(TEXT("RemoveUnusedOutputs: Failed to compile!")));
 		OutErrors = Data.Errors;
@@ -944,9 +949,16 @@ bool RemoveUnusedOutputs(FString& InOutSourceCode, const TArray<FString>& InUsed
 	return false;
 }
 
+bool RemoveUnusedOutputs(FString& InOutSourceCode, const TArray<FString>& InUsedOutputs, const TArray<FString>& InExceptions, FString& EntryPoint, TArray<FString>& OutErrors)
+{
+	const TArray<FStringView> UsedOutputs(MakeArrayView(InUsedOutputs));
+	const TArray<FStringView> Exceptions(MakeArrayView(InExceptions));
+	return RemoveUnusedOutputs(InOutSourceCode, UsedOutputs, Exceptions, {}, EntryPoint, OutErrors);
+}
+
 struct FRemoveUnusedInputs : FRemoveAlgorithm
 {
-	const TArray<FString>& UsedInputs;
+	const TConstArrayView<FStringView> UsedInputs;
 
 	struct FInputsBodyContext : FBodyContext
 	{
@@ -963,7 +975,7 @@ struct FRemoveUnusedInputs : FRemoveAlgorithm
 		}
 	};
 
-	FRemoveUnusedInputs(const TArray<FString>& InUsedInputs) :
+	FRemoveUnusedInputs(const TConstArrayView<FStringView> InUsedInputs) :
 		UsedInputs(InUsedInputs)
 	{
 	}
@@ -1369,11 +1381,15 @@ struct FRemoveUnusedInputs : FRemoveAlgorithm
 	}
 };
 
-bool RemoveUnusedInputs(FString& InOutSourceCode, const TArray<FString>& InInputs, FString& EntryPoint, TArray<FString>& OutErrors)
+bool RemoveUnusedInputs(FString& InOutSourceCode,
+	TConstArrayView<FStringView> InUsedInputs,
+	TConstArrayView<FScopedDeclarations> InScopedDeclarations,
+	FString& InOutEntryPoint,
+	TArray<FString>& OutErrors)
 {
 	FString DummyFilename(TEXT("/Engine/Private/RemoveUnusedInputs.usf"));
-	FRemoveUnusedInputs Data(InInputs);
-	Data.EntryPoint = EntryPoint;
+	FRemoveUnusedInputs Data(InUsedInputs);
+	Data.EntryPoint = InOutEntryPoint;
 	CrossCompiler::FCompilerMessages Messages;
 	auto Lambda = [&Data](CrossCompiler::FLinearAllocator* Allocator, CrossCompiler::TLinearArray<CrossCompiler::AST::FNode*>& ASTNodes)
 	{
@@ -1385,7 +1401,7 @@ bool RemoveUnusedInputs(FString& InOutSourceCode, const TArray<FString>& InInput
 			++i;
 		}
 	};
-	if (!CrossCompiler::Parser::Parse(InOutSourceCode, DummyFilename, Messages, Lambda))
+	if (!CrossCompiler::Parser::Parse(InOutSourceCode, DummyFilename, Messages, InScopedDeclarations, Lambda))
 	{
 		Data.Errors.Add(FString(TEXT("RemoveUnusedInputs: Failed to compile!")));
 		OutErrors = Data.Errors;
@@ -1405,7 +1421,7 @@ bool RemoveUnusedInputs(FString& InOutSourceCode, const TArray<FString>& InInput
 	{
 		InOutSourceCode += (TCHAR)'\n';
 		InOutSourceCode += Data.GeneratedCode;
-		EntryPoint = Data.EntryPoint;
+		InOutEntryPoint = Data.EntryPoint;
 
 		return true;
 	}
@@ -1414,6 +1430,12 @@ bool RemoveUnusedInputs(FString& InOutSourceCode, const TArray<FString>& InInput
 	return false;
 }
 
+bool RemoveUnusedInputs(FString& InOutSourceCode, const TArray<FString>& InInputs, FString& EntryPoint, TArray<FString>& OutErrors)
+{
+	const TArray<FStringView> Inputs(MakeArrayView(InInputs));
+	return RemoveUnusedInputs(InOutSourceCode, Inputs, {}, EntryPoint, OutErrors);
+}
+
 struct FConvertFP32ToFP16 {
 	FString Filename;
 	FString GeneratedCode;
@@ -1688,7 +1710,7 @@ bool ConvertFromFP32ToFP16(FString& InOutSourceCode, TArray<FString>& OutErrors)
 	FConvertFP32ToFP16 Data;
 	Data.Filename = DummyFilename;
 	Data.GeneratedCode = "";
-	if (!CrossCompiler::Parser::Parse(InOutSourceCode, DummyFilename, Messages, HlslParserCallbackWrapperFP32ToFP16, &Data))
+	if (!CrossCompiler::Parser::Parse(InOutSourceCode, DummyFilename, Messages, {}, HlslParserCallbackWrapperFP32ToFP16, &Data))
 	{
 		OutErrors.Add(FString(TEXT("ConvertFP32ToFP16: Failed to compile!")));
 		for (auto& Message : Messages.MessageList)
@@ -1728,7 +1750,7 @@ bool PrettyPrintHlslParser(FString& InOutSourceCode, TArray<FString>& OutErrors)
 	FString DummyFilename(TEXT("/Engine/Private/PrettyPrinter.usf"));
 	FString Out;
 	CrossCompiler::FCompilerMessages Messages;
-	if (!CrossCompiler::Parser::Parse(InOutSourceCode, DummyFilename, Messages, PrettyPrinter, &Out))
+	if (!CrossCompiler::Parser::Parse(InOutSourceCode, DummyFilename, Messages, {}, PrettyPrinter, &Out))
 	{
 		OutErrors.Add(FString(TEXT("PrettyPrintHlslParser: Failed to compile!")));
 		for (auto& Message : Messages.MessageList)
diff --git a/Engine/Source/Developer/ShaderCompilerCommon/Public/ShaderCompilerCommon.h b/Engine/Source/Developer/ShaderCompilerCommon/Public/ShaderCompilerCommon.h
--- a/Engine/Source/Developer/ShaderCompilerCommon/Public/ShaderCompilerCommon.h
+++ b/Engine/Source/Developer/ShaderCompilerCommon/Public/ShaderCompilerCommon.h
@@ -345,8 +345,37 @@ extern SHADERCOMPILERCOMMON_API void ParseHLSLTypeName(const TCHAR* SearchString
 // Processes TEXT macros
 extern SHADERCOMPILERCOMMON_API void TransformStringIntoCharacterArray(FString& PreprocessedShaderSource, TArray<FShaderDiagnosticData>* OutDiagnosticDatas=nullptr);
 
+// Structure to hold forward declarations for a specific scope/namespace chain for the HlslParser
+struct FScopedDeclarations
+{
+	FScopedDeclarations(TConstArrayView<FStringView> InScope, TConstArrayView<FStringView> InSymbols)
+		: Scope(InScope)
+		, Symbols(InSymbols)
+	{
+	}
+	TConstArrayView<FStringView> Scope;
+	TConstArrayView<FStringView> Symbols;
+};
+
+extern SHADERCOMPILERCOMMON_API bool RemoveUnusedOutputs(
+	FString& InOutSourceCode,
+	TConstArrayView<FStringView> InUsedOutputs,
+	TConstArrayView<FStringView> InExceptions,
+	TConstArrayView<FScopedDeclarations> InScopedDeclarations,
+	FString& InOutEntryPoint,
+	TArray<FString>& OutErrors
+);
+
 extern SHADERCOMPILERCOMMON_API bool RemoveUnusedOutputs(FString& InOutSourceCode, const TArray<FString>& InUsedOutputs, const TArray<FString>& InExceptions, FString& InOutEntryPoint, TArray<FString>& OutErrors);
 
+extern SHADERCOMPILERCOMMON_API bool RemoveUnusedInputs(
+	FString& InOutSourceCode,
+	TConstArrayView<FStringView> InUsedInputs,
+	TConstArrayView<FScopedDeclarations> InScopedDeclarations,
+	FString& InOutEntryPoint,
+	TArray<FString>& OutErrors
+);
+
 extern SHADERCOMPILERCOMMON_API bool RemoveUnusedInputs(FString& InOutSourceCode, const TArray<FString>& InUsedInputs, FString& InOutEntryPoint, TArray<FString>& OutErrors);
 
 extern SHADERCOMPILERCOMMON_API bool ConvertFromFP32ToFP16(FString& InOutSourceCode, TArray<FString>& OutErrors);
