Fixing crash that was hidden by hardcoded namespaces.

diff --git a/Engine/Source/Developer/ShaderCompilerCommon/Private/HlslExpressionParser.inl b/Engine/Source/Developer/ShaderCompilerCommon/Private/HlslExpressionParser.inl
--- a/Engine/Source/Developer/ShaderCompilerCommon/Private/HlslExpressionParser.inl
+++ b/Engine/Source/Developer/ShaderCompilerCommon/Private/HlslExpressionParser.inl
@@ -79,7 +79,7 @@ namespace CrossCompiler
 		{
 			for (auto& Child : Children)
 			{
-				if (!FCString::Strcmp(Child.Name, Namespace))
+				if (Child.Name && FCString::Strcmp(Child.Name, Namespace) == 0)
 				{
 					return &Child;
 				}
