D3D11 cross-compilation fixes:
- Patch HLSL to support SV_DepthGreaterEqual and SV_DepthLessEqual.
- Maintain order of shader input/output variables in cross-compiled HLSL.

#rnx

diff --git a/Engine/Source/Developer/ShaderCompilerCommon/Private/HlslccHeaderWriter.cpp b/Engine/Source/Developer/ShaderCompilerCommon/Private/HlslccHeaderWriter.cpp
--- a/Engine/Source/Developer/ShaderCompilerCommon/Private/HlslccHeaderWriter.cpp
+++ b/Engine/Source/Developer/ShaderCompilerCommon/Private/HlslccHeaderWriter.cpp
@@ -2,6 +2,7 @@
 
 #include "HlslccHeaderWriter.h"
 #include "ShaderConductorContext.h"
+#include "SpirvCommon.h"
 
 THIRD_PARTY_INCLUDES_START
 #include "spirv_reflect.h"
@@ -122,84 +123,9 @@ namespace CrossCompiler
 		return TypeName;
 	}
 
-	static const TCHAR* SpvBuiltinToString(const SpvBuiltIn BuiltIn)
-	{
-		switch (BuiltIn)
-		{
-		case SpvBuiltInPosition:					return TEXT("gl_Position");
-		case SpvBuiltInPointSize:					return TEXT("gl_PointSize");
-		case SpvBuiltInClipDistance:				return TEXT("gl_ClipDistance");
-		case SpvBuiltInCullDistance:				return TEXT("gl_CullDistance");
-		case SpvBuiltInVertexId:					return TEXT("gl_VertexID");
-		case SpvBuiltInInstanceId:					return TEXT("gl_InstanceID");
-		case SpvBuiltInPrimitiveId:					return TEXT("gl_PrimitiveID");
-		case SpvBuiltInInvocationId:				return TEXT("gl_InvocationID");
-		case SpvBuiltInLayer:						return TEXT("gl_Layer");
-		case SpvBuiltInViewportIndex:				return TEXT("gl_ViewportIndex");
-		case SpvBuiltInTessLevelOuter:				return TEXT("gl_TessLevelOuter");
-		case SpvBuiltInTessLevelInner:				return TEXT("gl_TessLevelInner");
-		case SpvBuiltInTessCoord:					return TEXT("gl_TessCoord");
-		case SpvBuiltInPatchVertices:				return TEXT("gl_PatchVertices");
-		case SpvBuiltInFragCoord:					return TEXT("gl_FragCoord");
-		case SpvBuiltInPointCoord:					return TEXT("gl_PointCoord");
-		case SpvBuiltInFrontFacing:					return TEXT("gl_FrontFacing");
-		case SpvBuiltInSampleId:					return TEXT("gl_SampleID");
-		case SpvBuiltInSamplePosition:				return TEXT("gl_SamplePosition");
-		case SpvBuiltInSampleMask:					return TEXT("gl_SampleMask");
-		case SpvBuiltInFragDepth:					return TEXT("gl_FragDepth");
-		case SpvBuiltInHelperInvocation:			return TEXT("gl_HelperInvocation");
-		case SpvBuiltInNumWorkgroups:				return TEXT("gl_NumWorkgroups");
-		case SpvBuiltInWorkgroupSize:				return TEXT("gl_WorkgroupSize");
-		case SpvBuiltInWorkgroupId:					return TEXT("gl_WorkgroupID");
-		case SpvBuiltInLocalInvocationId:			return TEXT("gl_LocalInvocationID");
-		case SpvBuiltInGlobalInvocationId:			return TEXT("gl_GlobalInvocationID");
-		case SpvBuiltInLocalInvocationIndex:		return TEXT("gl_LocalInvocationIndex");
-		case SpvBuiltInWorkDim:						return TEXT("gl_WorkDim");
-		case SpvBuiltInGlobalSize:					return TEXT("gl_GlobalSize");
-		case SpvBuiltInEnqueuedWorkgroupSize:		return TEXT("gl_EnqueuedWorkgroupSize");
-		case SpvBuiltInGlobalOffset:				return TEXT("gl_GlobalOffset");
-		case SpvBuiltInGlobalLinearId:				return TEXT("gl_GlobalLinearID");
-		case SpvBuiltInSubgroupSize:				return TEXT("gl_SubgroupSize");
-		case SpvBuiltInSubgroupMaxSize:				return TEXT("gl_SubgroupMaxSize");
-		case SpvBuiltInNumSubgroups:				return TEXT("gl_NumSubgroups");
-		case SpvBuiltInNumEnqueuedSubgroups:		return TEXT("gl_NumEnqueuedSubgroups");
-		case SpvBuiltInSubgroupId:					return TEXT("gl_SubgroupID");
-		case SpvBuiltInSubgroupLocalInvocationId:	return TEXT("gl_SubgroupLocalInvocationID");
-		case SpvBuiltInVertexIndex:					return TEXT("gl_VertexIndex");
-		case SpvBuiltInInstanceIndex:				return TEXT("gl_InstanceIndex");
-		case SpvBuiltInSubgroupEqMask:				return TEXT("gl_SubgroupEqMask");
-		case SpvBuiltInSubgroupGeMask:				return TEXT("gl_SubgroupGeMask");
-		case SpvBuiltInSubgroupGtMask:				return TEXT("gl_SubgroupGtMask");
-		case SpvBuiltInSubgroupLeMask:				return TEXT("gl_SubgroupLeMask");
-		case SpvBuiltInSubgroupLtMask:				return TEXT("gl_SubgroupLtMask");
-		case SpvBuiltInBaseVertex:					return TEXT("gl_BaseVertex");
-		case SpvBuiltInBaseInstance:				return TEXT("gl_BaseInstance");
-		case SpvBuiltInDrawIndex:					return TEXT("gl_DrawIndex");
-		case SpvBuiltInDeviceIndex:					return TEXT("gl_DeviceIndex");
-		case SpvBuiltInViewIndex:					return TEXT("gl_ViewIndex");
-
-		// Ray tracing
-		case SpvBuiltInLaunchIdKHR:					return TEXT("gl_LaunchIDEXT");
-		case SpvBuiltInLaunchSizeKHR:				return TEXT("gl_LaunchSizeEXT");
-		case SpvBuiltInInstanceCustomIndexKHR:		return TEXT("gl_InstanceCustomIndexEXT");
-		case SpvBuiltInRayGeometryIndexKHR:			return TEXT("gl_GeometryIndexEXT");
-		case SpvBuiltInWorldRayOriginKHR:			return TEXT("gl_WorldRayOriginEXT");
-		case SpvBuiltInWorldRayDirectionKHR:		return TEXT("gl_WorldRayDirectionEXT");
-		case SpvBuiltInObjectRayOriginKHR:			return TEXT("gl_ObjectRayOriginEXT");
-		case SpvBuiltInObjectRayDirectionKHR:		return TEXT("gl_ObjectRayDirectionEXT");
-		case SpvBuiltInRayTminKHR:					return TEXT("gl_RayTminEXT");
-		case SpvBuiltInRayTmaxKHR:					return TEXT("gl_RayTmaxEXT");
-		case SpvBuiltInIncomingRayFlagsKHR:			return TEXT("gl_IncomingRayFlagsEXT");
-		case SpvBuiltInHitKindKHR:					return TEXT("gl_HitKindEXT");
-		case SpvBuiltInObjectToWorldKHR:			return TEXT("gl_ObjectToWorldEXT");
-		case SpvBuiltInWorldToObjectKHR:			return TEXT("gl_WorldToObjectEXT");
-		}
-		return nullptr;
-	}
-
 	static FString ConvertAttributeToMetaDataSemantic(const ANSICHAR* AttributeName, const SpvBuiltIn BuiltIn, bool bIsInput)
 	{
-		if (const TCHAR* BuiltInName = SpvBuiltinToString(BuiltIn))
+		if (const TCHAR* BuiltInName = SpirvBuiltinToString(BuiltIn))
 		{
 			return FString(BuiltInName);
 		}
diff --git a/Engine/Source/Developer/ShaderCompilerCommon/Private/SpirvCommon.cpp b/Engine/Source/Developer/ShaderCompilerCommon/Private/SpirvCommon.cpp
--- a/Engine/Source/Developer/ShaderCompilerCommon/Private/SpirvCommon.cpp
+++ b/Engine/Source/Developer/ShaderCompilerCommon/Private/SpirvCommon.cpp
@@ -2,6 +2,81 @@
 
 #include "SpirvCommon.h"
 
+SHADERCOMPILERCOMMON_API const TCHAR* SpirvBuiltinToString(const SpvBuiltIn BuiltIn)
+{
+	switch (BuiltIn)
+	{
+	case SpvBuiltInPosition:					return TEXT("gl_Position");
+	case SpvBuiltInPointSize:					return TEXT("gl_PointSize");
+	case SpvBuiltInClipDistance:				return TEXT("gl_ClipDistance");
+	case SpvBuiltInCullDistance:				return TEXT("gl_CullDistance");
+	case SpvBuiltInVertexId:					return TEXT("gl_VertexID");
+	case SpvBuiltInInstanceId:					return TEXT("gl_InstanceID");
+	case SpvBuiltInPrimitiveId:					return TEXT("gl_PrimitiveID");
+	case SpvBuiltInInvocationId:				return TEXT("gl_InvocationID");
+	case SpvBuiltInLayer:						return TEXT("gl_Layer");
+	case SpvBuiltInViewportIndex:				return TEXT("gl_ViewportIndex");
+	case SpvBuiltInTessLevelOuter:				return TEXT("gl_TessLevelOuter");
+	case SpvBuiltInTessLevelInner:				return TEXT("gl_TessLevelInner");
+	case SpvBuiltInTessCoord:					return TEXT("gl_TessCoord");
+	case SpvBuiltInPatchVertices:				return TEXT("gl_PatchVertices");
+	case SpvBuiltInFragCoord:					return TEXT("gl_FragCoord");
+	case SpvBuiltInPointCoord:					return TEXT("gl_PointCoord");
+	case SpvBuiltInFrontFacing:					return TEXT("gl_FrontFacing");
+	case SpvBuiltInSampleId:					return TEXT("gl_SampleID");
+	case SpvBuiltInSamplePosition:				return TEXT("gl_SamplePosition");
+	case SpvBuiltInSampleMask:					return TEXT("gl_SampleMask");
+	case SpvBuiltInFragDepth:					return TEXT("gl_FragDepth");
+	case SpvBuiltInHelperInvocation:			return TEXT("gl_HelperInvocation");
+	case SpvBuiltInNumWorkgroups:				return TEXT("gl_NumWorkgroups");
+	case SpvBuiltInWorkgroupSize:				return TEXT("gl_WorkgroupSize");
+	case SpvBuiltInWorkgroupId:					return TEXT("gl_WorkgroupID");
+	case SpvBuiltInLocalInvocationId:			return TEXT("gl_LocalInvocationID");
+	case SpvBuiltInGlobalInvocationId:			return TEXT("gl_GlobalInvocationID");
+	case SpvBuiltInLocalInvocationIndex:		return TEXT("gl_LocalInvocationIndex");
+	case SpvBuiltInWorkDim:						return TEXT("gl_WorkDim");
+	case SpvBuiltInGlobalSize:					return TEXT("gl_GlobalSize");
+	case SpvBuiltInEnqueuedWorkgroupSize:		return TEXT("gl_EnqueuedWorkgroupSize");
+	case SpvBuiltInGlobalOffset:				return TEXT("gl_GlobalOffset");
+	case SpvBuiltInGlobalLinearId:				return TEXT("gl_GlobalLinearID");
+	case SpvBuiltInSubgroupSize:				return TEXT("gl_SubgroupSize");
+	case SpvBuiltInSubgroupMaxSize:				return TEXT("gl_SubgroupMaxSize");
+	case SpvBuiltInNumSubgroups:				return TEXT("gl_NumSubgroups");
+	case SpvBuiltInNumEnqueuedSubgroups:		return TEXT("gl_NumEnqueuedSubgroups");
+	case SpvBuiltInSubgroupId:					return TEXT("gl_SubgroupID");
+	case SpvBuiltInSubgroupLocalInvocationId:	return TEXT("gl_SubgroupLocalInvocationID");
+	case SpvBuiltInVertexIndex:					return TEXT("gl_VertexIndex");
+	case SpvBuiltInInstanceIndex:				return TEXT("gl_InstanceIndex");
+	case SpvBuiltInSubgroupEqMask:				return TEXT("gl_SubgroupEqMask");
+	case SpvBuiltInSubgroupGeMask:				return TEXT("gl_SubgroupGeMask");
+	case SpvBuiltInSubgroupGtMask:				return TEXT("gl_SubgroupGtMask");
+	case SpvBuiltInSubgroupLeMask:				return TEXT("gl_SubgroupLeMask");
+	case SpvBuiltInSubgroupLtMask:				return TEXT("gl_SubgroupLtMask");
+	case SpvBuiltInBaseVertex:					return TEXT("gl_BaseVertex");
+	case SpvBuiltInBaseInstance:				return TEXT("gl_BaseInstance");
+	case SpvBuiltInDrawIndex:					return TEXT("gl_DrawIndex");
+	case SpvBuiltInDeviceIndex:					return TEXT("gl_DeviceIndex");
+	case SpvBuiltInViewIndex:					return TEXT("gl_ViewIndex");
+
+	// Ray tracing
+	case SpvBuiltInLaunchIdKHR:					return TEXT("gl_LaunchIDEXT");
+	case SpvBuiltInLaunchSizeKHR:				return TEXT("gl_LaunchSizeEXT");
+	case SpvBuiltInInstanceCustomIndexKHR:		return TEXT("gl_InstanceCustomIndexEXT");
+	case SpvBuiltInRayGeometryIndexKHR:			return TEXT("gl_GeometryIndexEXT");
+	case SpvBuiltInWorldRayOriginKHR:			return TEXT("gl_WorldRayOriginEXT");
+	case SpvBuiltInWorldRayDirectionKHR:		return TEXT("gl_WorldRayDirectionEXT");
+	case SpvBuiltInObjectRayOriginKHR:			return TEXT("gl_ObjectRayOriginEXT");
+	case SpvBuiltInObjectRayDirectionKHR:		return TEXT("gl_ObjectRayDirectionEXT");
+	case SpvBuiltInRayTminKHR:					return TEXT("gl_RayTminEXT");
+	case SpvBuiltInRayTmaxKHR:					return TEXT("gl_RayTmaxEXT");
+	case SpvBuiltInIncomingRayFlagsKHR:			return TEXT("gl_IncomingRayFlagsEXT");
+	case SpvBuiltInHitKindKHR:					return TEXT("gl_HitKindEXT");
+	case SpvBuiltInObjectToWorldKHR:			return TEXT("gl_ObjectToWorldEXT");
+	case SpvBuiltInWorldToObjectKHR:			return TEXT("gl_WorldToObjectEXT");
+	}
+	return nullptr;
+}
+
 SHADERCOMPILERCOMMON_API void FindOffsetToSpirvEntryPoint(const FSpirv& Spirv, const ANSICHAR* EntryPointName, uint32& OutWordOffsetToEntryPoint, uint32& OutWordOffsetToMainName)
 {
 	// Iterate over all SPIR-V instructions until we have what we need
@@ -74,3 +149,134 @@ SHADERCOMPILERCOMMON_API const ANSICHAR* PatchSpirvEntryPointWithCRC(FSpirv& Spi
 
 	return reinterpret_cast<const ANSICHAR*>(Spirv.Data.GetData() + OffsetToEntryPoint);
 }
+
+struct FSpirvVariableIdAndName
+{
+	uint32 Id;
+	const ANSICHAR* AnsiName;
+	const TCHAR* WideName;
+};
+
+static void ParseSpirvVariableIDs(const FSpirv& Spirv, SpvStorageClass StorageClass, TArray<FSpirvVariableIdAndName>& OutVariables)
+{
+	for (FSpirvConstIterator Iter = Spirv.begin(); Iter != Spirv.end(); ++Iter)
+	{
+		switch (Iter.Opcode())
+		{
+		case SpvOpVariable:
+		{
+			const SpvStorageClass VariableStorageClass = Iter.OperandAs<SpvStorageClass>(3);
+			if (VariableStorageClass == StorageClass)
+			{
+				const uint32 VariableID = Iter.Operand(2);
+				OutVariables.Add({ VariableID, nullptr });
+			}
+		}
+		break;
+
+		case SpvOpFunction:
+		{
+			// Early exit with first function declaration
+			return;
+		}
+		break;
+		}
+	}
+}
+
+static void ParseSpirvVariableNames(const FSpirv& Spirv, TArray<FSpirvVariableIdAndName>& InOutVariables)
+{
+	for (FSpirvConstIterator Iter = Spirv.begin(); Iter != Spirv.end(); ++Iter)
+	{
+		switch (Iter.Opcode())
+		{
+		case SpvOpName:
+		{
+			// Iterate through all variables - there shouldn't be a lot of them when searching for a specific storage class (somewhere in the range 1-32).
+			const uint32 Target = Iter.Operand(1);
+			for (FSpirvVariableIdAndName& Variable : InOutVariables)
+			{
+				if (Variable.Id == Target)
+				{
+					Variable.AnsiName = Iter.OperandAsString(2);
+					break;
+				}
+			}
+		}
+		break;
+
+		case SpvOpDecorate:
+		{
+			// Is this a decoration for a built-in variable?
+			const SpvDecoration Decoration = Iter.OperandAs<SpvDecoration>(2);
+			if (Decoration == SpvDecorationBuiltIn)
+			{
+				const uint32 Target = Iter.Operand(1);
+				for (FSpirvVariableIdAndName& Variable : InOutVariables)
+				{
+					if (Variable.Id == Target)
+					{
+						const SpvBuiltIn BuiltIn = Iter.OperandAs<SpvBuiltIn>(3);
+						Variable.WideName = SpirvBuiltinToString(BuiltIn);
+						break;
+					}
+				}
+			}
+		}
+		break;
+
+		case SpvOpFunction:
+		{
+			// Early exit with first function declaration
+			return;
+		}
+		break;
+		}
+	}
+}
+
+// Replaces '.' characters from the specified SPIR-V name with '_'.
+static FString SanitizeSpirvName(const ANSICHAR* Name)
+{
+	FString NameStr = ANSI_TO_TCHAR(Name);
+	for (TCHAR& Chr : NameStr)
+	{
+		if (Chr == TEXT('.'))
+		{
+			Chr = TEXT('_');
+		}
+	}
+	return NameStr;
+}
+
+SHADERCOMPILERCOMMON_API void ParseSpirvGlobalVariables(const FSpirv& Spirv, SpvStorageClass StorageClass, TArray<FString>& OutVariableNames)
+{
+	checkf(StorageClass != SpvStorageClassFunction, TEXT("StorageClass must not be SpvStorageClassFunction"));
+
+	// Parse variable IDs and names from SPIR-V module
+	TArray<FSpirvVariableIdAndName> IntermediateVariables;
+	ParseSpirvVariableIDs(Spirv, StorageClass, IntermediateVariables);
+
+	// Now find names for variable IDs, since OpName comes first in SPIR-V module but assigns their values to target variables that are declared later
+	ParseSpirvVariableNames(Spirv, IntermediateVariables);
+
+	// Convert variable names to FString output container
+	OutVariableNames.Reserve(IntermediateVariables.Num());
+	for (const FSpirvVariableIdAndName& Variable : IntermediateVariables)
+	{
+		if (const ANSICHAR* AnsiName = Variable.AnsiName)
+		{
+			// Names in SPIR-V module may contain invalid characters for high-level source languages, e.g. "in.var.TEXCOORD0"
+			OutVariableNames.Add(SanitizeSpirvName(AnsiName));
+		}
+		else if (const TCHAR* WideName = Variable.WideName)
+		{
+			// No sanitization of this name is necessary; this is supposed to come from SpirvBuiltinToString().
+			OutVariableNames.Add(WideName);
+		}
+		else
+		{
+			checkf(0, TEXT("FSpirvVariableIdAndName entry with invalid name fields (ID = %u)"), Variable.Id);
+		}
+	}
+}
diff --git a/Engine/Source/Developer/ShaderCompilerCommon/Public/SpirvCommon.h b/Engine/Source/Developer/ShaderCompilerCommon/Public/SpirvCommon.h
--- a/Engine/Source/Developer/ShaderCompilerCommon/Public/SpirvCommon.h
+++ b/Engine/Source/Developer/ShaderCompilerCommon/Public/SpirvCommon.h
@@ -187,10 +187,16 @@ public:
 	}
 };
 
+/** Returns the string representation of the specified SPIR-V built-in variable, e.g. SpvBuiltInPosition to TEXT("gl_Position"). Returns nullptr for invalid value. */
+SHADERCOMPILERCOMMON_API const TCHAR* SpirvBuiltinToString(const SpvBuiltIn BuiltIn);
+
 /** Returns word offset to the entry point (OpEntryPoint) and name of the entry point (OpName). If the respective entry has not been found, the output offset is 0. */
-extern SHADERCOMPILERCOMMON_API void FindOffsetToSpirvEntryPoint(const FSpirv& Spirv, const ANSICHAR* EntryPointName, uint32& OutWordOffsetToEntryPoint, uint32& OutWordOffsetToMainName);
+SHADERCOMPILERCOMMON_API void FindOffsetToSpirvEntryPoint(const FSpirv& Spirv, const ANSICHAR* EntryPointName, uint32& OutWordOffsetToEntryPoint, uint32& OutWordOffsetToMainName);
 
 /** Renames the fixed-size entry point name (which must be "main_00000000_00000000") to the formatted name including a CRC over the module. Returns the new entry point name. */
-extern SHADERCOMPILERCOMMON_API const ANSICHAR* PatchSpirvEntryPointWithCRC(FSpirv& Spirv, uint32& OutCRC);
+SHADERCOMPILERCOMMON_API const ANSICHAR* PatchSpirvEntryPointWithCRC(FSpirv& Spirv, uint32& OutCRC);
+
+/** Parse global variables (anything but function local variables) of specified storage class from the specified SPIR-V module. */
+SHADERCOMPILERCOMMON_API void ParseSpirvGlobalVariables(const FSpirv& Spirv, SpvStorageClass StorageClass, TArray<FString>& OutVariableNames);
 
 
diff --git a/Engine/Source/Developer/Windows/ShaderFormatD3D/Private/D3DShaderCompiler.cpp b/Engine/Source/Developer/Windows/ShaderFormatD3D/Private/D3DShaderCompiler.cpp
--- a/Engine/Source/Developer/Windows/ShaderFormatD3D/Private/D3DShaderCompiler.cpp
+++ b/Engine/Source/Developer/Windows/ShaderFormatD3D/Private/D3DShaderCompiler.cpp
@@ -506,8 +506,75 @@ static void PatchSpirvForPrecompilation(FSpirv& Spirv)
 	}
 }
 
+// @param StageVariablesStorageClass Must be SpvStorageClassOutput for vertex shaders and SpvStorageClassInput for pixel shaders.
+static bool PatchHlslWithReorderedIOVariables(FString& HlslSourceString, const FSpirv& Spirv, SpvStorageClass StageVariablesStorageClass)
+{
+	check(StageVariablesStorageClass == SpvStorageClassInput || StageVariablesStorageClass == SpvStorageClassOutput);
+
+	// Find declaration struct for stage variables
+	const FStringView StageVariableDeclarationName = (StageVariablesStorageClass == SpvStorageClassInput ? TEXT("SPIRV_Cross_Input") : TEXT("SPIRV_Cross_Output"));
+	const int32 StageVariableDeclarationBegin = HlslSourceString.Find(StageVariableDeclarationName, ESearchCase::CaseSensitive);
+	if (StageVariableDeclarationBegin == INDEX_NONE)
+	{
+		return false;
+	}
+
+	const int32 StageVariableDelcarationBlockBegin = HlslSourceString.Find(TEXT("{"), ESearchCase::CaseSensitive, ESearchDir::FromStart, StageVariableDeclarationBegin + StageVariableDeclarationName.Len());
+	if (StageVariableDelcarationBlockBegin == INDEX_NONE)
+	{
+		return false;
+	}
+
+	const int32 StageVariableDelcarationBlockEnd = HlslSourceString.Find(TEXT("}"), ESearchCase::CaseSensitive, ESearchDir::FromStart, StageVariableDelcarationBlockBegin + 1);
+	if (StageVariableDelcarationBlockEnd == INDEX_NONE)
+	{
+		return false;
+	}
+
+	// Parse declaration struct for stage variables into array of individual lines
+	const FString StageVariableDeclarationSource = HlslSourceString.Mid(StageVariableDelcarationBlockBegin + 1, StageVariableDelcarationBlockEnd - (StageVariableDelcarationBlockBegin + 1));
+
+	TArray<FString> StageVariableDeclarationLines;
+	StageVariableDeclarationSource.ParseIntoArrayLines(StageVariableDeclarationLines);
+
+	// Parse variable names from SPIR-V input
+	TArray<FString> Variables;
+	ParseSpirvGlobalVariables(Spirv, StageVariablesStorageClass, Variables);
+
+	if (Variables.Num() != StageVariableDeclarationLines.Num())
+	{
+		// Failed to match SPIR-V variables to SPIRV-Cross generated source
+		return false;
+	}
+
+	// Re-arrange source lines of stage variable declarations
+	FString SortedStageVariableDeclarationSource = TEXT("\n");
+
+	for (const FString& Variable : Variables)
+	{
+		for (FString& SourceLine : StageVariableDeclarationLines)
+		{
+			if (SourceLine.Find(Variable, ESearchCase::CaseSensitive) != INDEX_NONE)
+			{
+				// Append source line for current variable at the end of sorted declaration string.
+				// Then empty this source line to avoid unnecessary string comparisons for next variables.
+				SortedStageVariableDeclarationSource += SourceLine;
+				SortedStageVariableDeclarationSource += TEXT('\n');
+				SourceLine.Empty();
+				break;
+			}
+		}
+	}
+
+	// Replace old declaration with sorted one
+	HlslSourceString.RemoveAt(StageVariableDelcarationBlockBegin + 1, StageVariableDelcarationBlockEnd - (StageVariableDelcarationBlockBegin + 1));
+	HlslSourceString.InsertAt(StageVariableDelcarationBlockBegin + 1, SortedStageVariableDeclarationSource);
+
+	return true;
+}
+
 // @todo-lh: use ANSI string class whenever UE core gets one
-static void PatchHlslForPrecompilation(TArray<ANSICHAR>& HlslSource)
+static void PatchHlslForPrecompilation(TArray<ANSICHAR>& HlslSource, const EShaderFrequency Frequency, const FSpirv& Spirv)
 {
 	FString HlslSourceString;
 
@@ -549,6 +616,32 @@ static void PatchHlslForPrecompilation(TArray<ANSICHAR>& HlslSource)
 		}
 	}
 
+	if (Frequency == SF_Vertex)
+	{
+		// Ensure order of output variables remains the same as declared in SPIR-V input
+		PatchHlslWithReorderedIOVariables(HlslSourceString, Spirv, SpvStorageClassOutput);
+	}
+	else if (Frequency == SF_Pixel)
+	{
+		// Patch internal error when SV_DepthLessEqual or SV_DepthGreaterEqual is specified in a pixel shader output. This is to prevent the following internal error:
+		//	error X8000 : D3D11 Internal Compiler Error : Invalid Bytecode : Interpolation mode for PS input position must be
+		//				  linear_noperspective_centroid or linear_noperspective_sample when outputting oDepthGE or oDepthLE and
+		//				  not running at sample frequency(which is forced by inputting SV_SampleIndex or declaring an input linear_sample or linear_noperspective_sample).
+		if (HlslSourceString.Find(TEXT("SV_DepthLessEqual"), ESearchCase::CaseSensitive) != INDEX_NONE ||
+			HlslSourceString.Find(TEXT("SV_DepthGreaterEqual"), ESearchCase::CaseSensitive) != INDEX_NONE)
+		{
+			// Ensure the interpolation mode is linear_noperspective_sample by adding "sample" specifier to one of the input-interpolators that have a floating-point type
+			const int32 FragCoordStringPosition = HlslSourceString.Find(TEXT("float4 gl_FragCoord : SV_Position"), ESearchCase::CaseSensitive);
+			if (FragCoordStringPosition != INDEX_NONE)
+			{
+				HlslSourceString.InsertAt(FragCoordStringPosition, TEXT("sample "));
+			}
+		}
+
+		// Ensure order of input variables remains the same as declared in SPIR-V input
+		PatchHlslWithReorderedIOVariables(HlslSourceString, Spirv, SpvStorageClassInput);
+	}
+
 	// Return new HLSL source
 	HlslSource.SetNum(HlslSourceString.Len() + 1);
 	FMemory::Memcpy(HlslSource.GetData(), TCHAR_TO_ANSI(*HlslSourceString), HlslSourceString.Len());
@@ -695,7 +788,7 @@ static bool CompileAndProcessD3DShaderFXCExt(
 			}
 
 			// Patch HLSL for workarounds to prevent potential additional FXC failures
-			PatchHlslForPrecompilation(CrossCompiledSource);
+			PatchHlslForPrecompilation(CrossCompiledSource, Frequency, Spirv);
 
 			if (bDumpDebugInfo && CrossCompiledSource.Num() > 1)
 			{
