ShaderMinifier - Optimize BuildLineBreakMap, SkipUntilNonIdentifierCharacter, ExtractIdentifiers and few other bits

* Generate line directive map during parsing
* Use xxhash instead of cityhash
* Avoid FStringView operator [] in ExtractIdentifiers
* Use TArray Reset instead of Empty where appropriate
* Vectorize SkipUntilNonIdentifierCharacter
* Vectorize ExtractBlock
* Vectorize line ending search
* Use FString directly instead of FStringBuilderBase for the final output to avoid a copy

diff --git a/Engine/Source/Developer/ShaderCompilerCommon/Private/ShaderMinifier.cpp b/Engine/Source/Developer/ShaderCompilerCommon/Private/ShaderMinifier.cpp
--- a/Engine/Source/Developer/ShaderCompilerCommon/Private/ShaderMinifier.cpp
+++ b/Engine/Source/Developer/ShaderCompilerCommon/Private/ShaderMinifier.cpp
@@ -3,12 +3,19 @@
 #include "ShaderMinifier.h"
 
 #include "HAL/PlatformTime.h"
-#include "Hash/CityHash.h"
+#include "Hash/xxhash.h"
 #include "Logging/LogMacros.h"
 #include "Misc/AutomationTest.h"
 #include "String/Find.h"
 #include "Algo/BinarySearch.h"
 
+#define UE_SHADER_MINIFIER_SSE (PLATFORM_CPU_X86_FAMILY && PLATFORM_ENABLE_VECTORINTRINSICS && PLATFORM_ALWAYS_HAS_SSE4_2)
+
+#if UE_SHADER_MINIFIER_SSE
+#include <emmintrin.h>
+#include <nmmintrin.h>
+#endif
+
 DEFINE_LOG_CATEGORY_STATIC(LogShaderMinifier, Log, All);
 
 // TODO:
@@ -78,30 +85,91 @@ static bool StartsWith(FStringView Source, FStringView Prefix)
 	return Equals(SourceView, Prefix);
 }
 
-static bool IsSpace(TCHAR C)
+struct FCharacterFlags
 {
-	switch (C)
+	static constexpr uint8 None = 0;
+
+	static constexpr uint8 Letter     = 1 << 0;
+	static constexpr uint8 Number     = 1 << 1;
+	static constexpr uint8 Underscore = 1 << 2;
+	static constexpr uint8 Space      = 1 << 3;
+	static constexpr uint8 Special    = 1 << 4;
+
+	static constexpr uint8 PossibleIdentifierMask = Letter | Number | Underscore;
+
+	FCharacterFlags()
 	{
-	default:
-		return false;
-	case TCHAR(' '):
-	case TCHAR('\f'):
-	case TCHAR('\r'):
-	case TCHAR('\n'):
-	case TCHAR('\t'):
-	case TCHAR('\v'):
-		return true;
+		for (char C = '0'; C <= '9'; ++C)
+		{
+			Flags[C] |= Number;
+		}
+
+		for (char C = 'a'; C <= 'z'; ++C)
+		{
+			Flags[C] |= Letter;
+		}
+
+		for (char C = 'A'; C <= 'Z'; ++C)
+		{
+			Flags[C] |= Letter;
+		}
+
+		Flags['_'] |= Underscore;
+
+		Flags[' ']  |= Space;
+		Flags['\f'] |= Space;
+		Flags['\r'] |= Space;
+		Flags['\n'] |= Space;
+		Flags['\t'] |= Space;
+		Flags['\v'] |= Space;
+
+		Flags['\\'] |= Special;
+		Flags['/']  |= Special;
+		Flags['#']  |= Special;
+		Flags['{']  |= Special;
+		Flags['}']  |= Special;
+		Flags['(']  |= Special;
+		Flags[')']  |= Special;
+	}
+
+	bool IsSpace(TCHAR C) const
+	{
+		return (Flags[uint8(C)] & Space) != 0;
+	}
+
+	bool IsNumber(TCHAR C) const
+	{
+		return (Flags[uint8(C)] & Number) != 0;
+	}
+
+	bool IsPossibleIdentifierCharacter(TCHAR C) const
+	{
+		return (Flags[uint8(C)] & PossibleIdentifierMask) != 0;
+	}
+
+	bool IsSpecial(TCHAR C) const
+	{
+		return (Flags[uint8(C)] & Special) != 0;
 	}
+
+	uint8 Flags[256] = {};
+};
+
+static const FCharacterFlags GCharacterFlags;
+
+static bool IsSpace(TCHAR C)
+{
+	return GCharacterFlags.IsSpace(C);
 }
 
 static bool IsNumber(TCHAR C)
 {
-	return C >= '0' && C <= '9';
+	return GCharacterFlags.IsNumber(C);
 }
 
 static bool IsPossibleIdentifierCharacter(TCHAR C)
 {
-	return (C >= '0' && C <= '9') || (C >= 'a' && C <= 'z') || (C >= 'A' && C <= 'Z') || C == '_';
+	return GCharacterFlags.IsPossibleIdentifierCharacter(C);
 }
 
 static FStringView ExtractOperator(FStringView Source)
@@ -172,11 +240,31 @@ static FStringView ExtractOperator(FStringView Source)
 	return Result;
 }
 
-static FStringView SkipUntilNonIdentifierCharacter(FStringView Source) 
+static FStringView SkipUntilNonIdentifierCharacter(FStringView Source)
 {
 	const int32 SourceLen = Source.Len();
 	int32 Cursor = 0;
 	const TCHAR* SourceData = Source.GetData();
+
+#if UE_SHADER_MINIFIER_SSE
+	{
+		const int32 AlignedLen = SourceLen & (~7); // align down to multiple of 8 TCHAR-s
+		const __m128i NeedleVec = _mm_setr_epi16(L'0', L'9', L'a', L'z', L'A', L'Z', L'_', L'_');
+		while (Cursor < AlignedLen)
+		{
+			__m128i Chunk = _mm_loadu_si128(reinterpret_cast<const __m128i*>(SourceData + Cursor));
+			constexpr int32 Mode = _SIDD_UWORD_OPS | _SIDD_CMP_RANGES | _SIDD_MASKED_NEGATIVE_POLARITY;
+			const int32 CompareResult = _mm_cmpistrc(NeedleVec, Chunk, Mode);
+			if (CompareResult)
+			{
+				Cursor += _mm_cmpistri(NeedleVec, Chunk, Mode);
+				return FStringView(SourceData + Cursor, SourceLen - Cursor);
+			}
+			Cursor += 8;
+		}
+	}
+#endif // UE_SHADER_MINIFIER_SSE
+
 	while (Cursor < SourceLen)
 	{
 		if (!IsPossibleIdentifierCharacter(SourceData[Cursor]))
@@ -209,6 +297,26 @@ static FStringView SkipSpace(FStringView Source)
 	const int32 SourceLen = Source.Len();
 	int32 Cursor = 0;
 	const TCHAR* SourceData = Source.GetData();
+
+#if UE_SHADER_MINIFIER_SSE
+	{
+		const int32 AlignedLen = SourceLen & (~7); // align down to multiple of 8 TCHAR-s
+		const __m128i NeedleVec = _mm_setr_epi16(L' ', L'\f', L'\r', L'\n', L'\t', L'\v', 0, 0);
+		while (Cursor < AlignedLen)
+		{
+			__m128i Chunk = _mm_loadu_si128(reinterpret_cast<const __m128i*>(SourceData + Cursor));
+			constexpr int32 Mode = _SIDD_UWORD_OPS | _SIDD_CMP_EQUAL_ANY | _SIDD_MASKED_NEGATIVE_POLARITY;
+			const int32 CompareResult = _mm_cmpistrc(NeedleVec, Chunk, Mode);
+			if (CompareResult)
+			{
+				Cursor += _mm_cmpistri(NeedleVec, Chunk, Mode);
+				return FStringView(SourceData + Cursor, SourceLen - Cursor);
+			}
+			Cursor += 8;
+		}
+	}
+#endif // UE_SHADER_MINIFIER_SSE
+
 	while (Cursor < SourceLen)
 	{
 		if (!IsSpace(SourceData[Cursor]))
@@ -275,9 +383,21 @@ static FStringView ExtractBlock(FStringView Source, TCHAR DelimBegin, TCHAR Deli
 	int32 Stack  = 0;
 	const int32 SourceLen = Source.Len();
 	const TCHAR* SourceData = Source.GetData();
-	for (int32 I = 0; I < SourceLen; ++I)
+
+	int32 Cursor = 0;
+
+	enum class EStatus
+	{
+		Finished,
+		Continue,
+	};
+
+	EStatus Status = EStatus::Continue;
+
+	auto ProcessCharacter = [&Stack, &PosEnd, SourceData, DelimBegin, DelimEnd](int32 Cursor) -> EStatus
 	{
-		TCHAR C = SourceData[I];
+		TCHAR C = SourceData[Cursor];
+
 		if (C == DelimBegin)
 		{
 			Stack++;
@@ -287,17 +407,51 @@ static FStringView ExtractBlock(FStringView Source, TCHAR DelimBegin, TCHAR Deli
 			if (Stack == 0)
 			{
 				// delimiter mismatch
-				break;
+				return EStatus::Finished;
 			}
 
 			Stack--;
 
 			if (Stack == 0)
 			{
-				PosEnd = I;
-				break;
+				PosEnd = Cursor;
+				return EStatus::Finished;
 			}
 		}
+
+		return EStatus::Continue;
+	};
+
+#if UE_SHADER_MINIFIER_SSE
+	{
+		const int32 AlignedLen = SourceLen & (~7); // align down to multiple of 8 TCHAR-s
+		const __m128i NeedleVec = _mm_setr_epi16(DelimBegin, DelimEnd, 0, 0, 0, 0, 0, 0);
+		while (Cursor < AlignedLen && Status != EStatus::Finished)
+		{
+			__m128i Chunk = _mm_loadu_si128(reinterpret_cast<const __m128i*>(SourceData + Cursor));
+			constexpr int32 Mode = _SIDD_UWORD_OPS | _SIDD_CMP_EQUAL_ANY | _SIDD_MOST_SIGNIFICANT;
+			const int32 CompareResult = _mm_cmpistrc(NeedleVec, Chunk, Mode);
+			if (CompareResult)
+			{
+				__m128i MaskVec = _mm_cmpistrm(NeedleVec, Chunk, Mode);
+				uint32 Mask = _mm_movemask_epi8(MaskVec);
+				while (Mask != 0 && Status != EStatus::Finished)
+				{
+					const uint32 BitIndex = FMath::CountTrailingZeros(Mask);
+					const uint32 ChunkCharIndex = BitIndex / 2;
+					Status = ProcessCharacter(Cursor + ChunkCharIndex);
+					Mask &= ~(3 << BitIndex);
+				}
+			}
+			Cursor += 8;
+		}
+	}
+#endif // UE_SHADER_MINIFIER_SSE
+
+	while (Cursor < SourceLen && Status != EStatus::Finished)
+	{
+		Status = ProcessCharacter(Cursor);
+		++Cursor;
 	}
 
 	if (Stack == 0 && PosEnd != INDEX_NONE)
@@ -331,10 +485,35 @@ enum class EBlockType : uint8 {
 
 struct FCodeBlock
 {
-	EBlockType  Type = EBlockType::Unknown;
-	FStringView Code;
+	const TCHAR* CodePtr = nullptr;
+	int32        CodeLen = 0;
+	EBlockType   Type = EBlockType::Unknown;
+	uint8        Padding[3] = {};
+
+	operator FStringView () const
+	{
+		return GetCode();
+	}
+
+	bool operator == (const FStringView S) const
+	{
+		return Equals(GetCode(), S);
+	}
+
+	void SetCode(FStringView Code)
+	{
+		CodePtr = Code.GetData();
+		CodeLen = Code.Len();
+	}
+
+	FStringView GetCode() const
+	{
+		return FStringView(CodePtr, CodeLen);
+	}
 };
 
+static_assert(sizeof(FCodeBlock) == 16, "Unexpected FCodeBlock size");
+
 enum class ECodeChunkType {
 	Unknown,
 	Struct,
@@ -372,10 +551,12 @@ struct FNamespace
 	TArray<FStringView> Stack; // i.e. [Foo, Bar, Baz]
 };
 
+using FCodeBlockArray = TArray<FCodeBlock, TInlineAllocator<6>>;
+
 struct FCodeChunk
 {
 	ECodeChunkType Type = ECodeChunkType::Unknown;
-	TArray<FCodeBlock> Blocks;
+	FCodeBlockArray Blocks;
 
 	int32 Namespace = INDEX_NONE; // Unique namespace ID (INDEX_NONE = global)
 
@@ -390,7 +571,7 @@ struct FCodeChunk
 		{
 			if (Block.Type == InType)
 			{
-				return Block.Code;
+				return Block;
 			}
 		}
 		return {};
@@ -407,8 +588,8 @@ struct FCodeChunk
 		{
 			const FCodeBlock& FirstBlock = Blocks[0];
 			const FCodeBlock& LastBlock = Blocks[Blocks.Num()-1];
-			const TCHAR* Begin = FirstBlock.Code.GetData();
-			const TCHAR* End = LastBlock.Code.GetData() + LastBlock.Code.Len();
+			const TCHAR* Begin = FirstBlock.CodePtr;
+			const TCHAR* End = LastBlock.CodePtr + LastBlock.CodeLen;
 			return FStringView(Begin, int32(End-Begin));
 		}
 	}
@@ -419,6 +600,7 @@ struct FParsedShader
 	FStringView Source;
 	TArray<FCodeChunk> Chunks;
 	TArray<FNamespace> Namespaces;
+	TArray<FStringView> LineDirectives;
 };
 
 struct FNamespaceTracker
@@ -477,7 +659,7 @@ static FParsedShader ParseShader(FStringView InSource, FDiagnostics& Output)
 	Result.Source = InSource;
 
 	FStringView		   Source = InSource;
-	TArray<FCodeBlock> PendingBlocks;
+	FCodeBlockArray    PendingBlocks;
 	TArray<FCodeChunk> Chunks;
 
 	ECodeChunkType ChunkType = ECodeChunkType::Unknown;
@@ -513,7 +695,7 @@ static FParsedShader ParseShader(FStringView InSource, FDiagnostics& Output)
 	{
 		FCodeBlock NewBlock;
 		NewBlock.Type = Type;
-		NewBlock.Code = Code;
+		NewBlock.SetCode(Code);
 		PendingBlocks.Push(NewBlock);
 	};
 
@@ -687,7 +869,7 @@ static FParsedShader ParseShader(FStringView InSource, FDiagnostics& Output)
 			bFoundIdentifier = false;
 			bFoundAssignment = false;
 
-			PendingBlocks.Empty();
+			PendingBlocks.Reset();
 		}
 	};
 
@@ -700,104 +882,123 @@ static FParsedShader ParseShader(FStringView InSource, FDiagnostics& Output)
 			break;
 		}
 
-		if (StartsWith(Source, TEXTVIEW("//")))
-		{
-			FStringView Remainder = SkipUntilNextLine(Source);
+		const TCHAR FirstChar = *Source.GetData();
 
-			// Save comment lines that are outside of blocks
-			if (PendingBlocks.IsEmpty())
+		if (GCharacterFlags.IsSpecial(FirstChar))
+		{
+			if (FirstChar == '/')
 			{
-				FStringView Block = SubStrView(Source, 0, Source.Len() - Remainder.Len());
-				AddBlock(EBlockType::Unknown, Block);
-				ChunkType = ECodeChunkType::CommentLine;
-				FinalizeChunk();
-			}
+				if (StartsWith(Source, TEXTVIEW("//")))
+				{
+					FStringView Remainder = SkipUntilNextLine(Source);
 
-			Source = Remainder;
+					// Save comment lines that are outside of blocks
+					if (PendingBlocks.IsEmpty())
+					{
+						FStringView Block = SubStrView(Source, 0, Source.Len() - Remainder.Len());
+						AddBlock(EBlockType::Unknown, Block);
+						ChunkType = ECodeChunkType::CommentLine;
+						FinalizeChunk();
+					}
 
-			continue;
-		}
-		else if (StartsWith(Source, TEXTVIEW("#line")))
-		{
-			Source = SkipUntilNextLine(Source);
-			continue;
-		}
-		else if (StartsWith(Source, TEXTVIEW("#pragma")))
-		{
-			FStringView Remainder = SkipUntilNextLine(Source);
-			FStringView Block = SubStrView(Source, 0, Source.Len() - Remainder.Len());
-			AddBlock(EBlockType::Directive, Block);
-			ChunkType = ECodeChunkType::Pragma;
-			FinalizeChunk();
-			Source = Remainder;
-			continue;
-		}
-		else if (StartsWith(Source, TEXTVIEW("#define")))
-		{
-			// TODO: handle `\` new lines in defines
-			FStringView Remainder = SkipUntilNextLine(Source);
-			FStringView Block = SubStrView(Source, 0, Source.Len() - Remainder.Len());
-			AddBlock(EBlockType::Directive, Block);
-			ChunkType = ECodeChunkType::Define;
-			FinalizeChunk();
-			Source = Remainder;
-			continue;
-		}
-		else if (StartsWith(Source, TEXTVIEW("#if 0")))
-		{
-			Source = SkipUntilStr(Source, TEXTVIEW("#endif"));
-			if (Source.Len() >= 6)
-			{
-				Source = SubStrView(Source, 6);
-			}
-			continue;
-		}
-		else if (StartsWith(Source, TEXTVIEW("/*")))
-		{
-			Source = SkipUntilStr(Source, TEXTVIEW("*/"));
-			if (Source.Len() >= 2)
-			{
-				Source = SubStrView(Source, 2);
+					Source = Remainder;
+
+					continue;
+				}
+				else if (StartsWith(Source, TEXTVIEW("/*")))
+				{
+					Source = SkipUntilStr(Source, TEXTVIEW("*/"));
+					if (Source.Len() >= 2)
+					{
+						Source = SubStrView(Source, 2);
+					}
+					continue;
+				}
 			}
-			continue;
-		}
-		else if (PendingBlocks.IsEmpty() && StartsWith(Source, TEXTVIEW("{")))
-		{
-			if (ChunkType == ECodeChunkType::Namespace)
+			else if (FirstChar == '#')
 			{
-				if (PendingNamespace.IsEmpty())
+				if (StartsWith(Source, TEXTVIEW("#line")))
 				{
-					AddDiagnostic(Output.Errors, TEXT("HLSL does not support anonymous namespaces"));
-					break;
+					FStringView Remainder = SkipUntilNextLine(Source);
+					FStringView Block = SubStrView(Source, 0, Source.Len() - Remainder.Len());
+					Result.LineDirectives.Add(Block);
+					Source = Remainder;
+					continue;
 				}
-				else
+				else if (StartsWith(Source, TEXTVIEW("#pragma")))
 				{
-					NamespaceTracker.Push(PendingNamespace);
-					ChunkType = ECodeChunkType::Unknown;
-					PendingNamespace = {};
-					Source = Source.Mid(1);
+					FStringView Remainder = SkipUntilNextLine(Source);
+					FStringView Block = SubStrView(Source, 0, Source.Len() - Remainder.Len());
+					AddBlock(EBlockType::Directive, Block);
+					ChunkType = ECodeChunkType::Pragma;
+					FinalizeChunk();
+					Source = Remainder;
+					continue;
+				}
+				else if (StartsWith(Source, TEXTVIEW("#define")))
+				{
+					// TODO: handle `\` new lines in defines
+					FStringView Remainder = SkipUntilNextLine(Source);
+					FStringView Block = SubStrView(Source, 0, Source.Len() - Remainder.Len());
+					AddBlock(EBlockType::Directive, Block);
+					ChunkType = ECodeChunkType::Define;
+					FinalizeChunk();
+					Source = Remainder;
+					continue;
+				}
+				else if (StartsWith(Source, TEXTVIEW("#if 0")))
+				{
+					Source = SkipUntilStr(Source, TEXTVIEW("#endif"));
+					if (Source.Len() >= 6)
+					{
+						Source = SubStrView(Source, 6);
+					}
+					continue;
 				}
-				continue;
-			}
-			else
-			{
-				AddDiagnostic(Output.Errors, TEXT("Expected token '{'"));
-			}
-			continue;
-		}
-		else if (PendingBlocks.IsEmpty() && StartsWith(Source, TEXTVIEW("}")))
-		{
-			if (NamespaceTracker.Pop())
-			{
-				Source = Source.Mid(1);
-				continue;
 			}
-			else
+			else if (PendingBlocks.IsEmpty() && (FirstChar == '{' || FirstChar == '}'))
 			{
-				AddDiagnostic(Output.Errors, TEXT("Expected token '}'"));
+				if (StartsWith(Source, TEXTVIEW("{")))
+				{
+					if (ChunkType == ECodeChunkType::Namespace)
+					{
+						if (PendingNamespace.IsEmpty())
+						{
+							AddDiagnostic(Output.Errors, TEXT("HLSL does not support anonymous namespaces"));
+							break;
+						}
+						else
+						{
+							NamespaceTracker.Push(PendingNamespace);
+							ChunkType = ECodeChunkType::Unknown;
+							PendingNamespace = {};
+							Source = Source.Mid(1);
+						}
+						continue;
+					}
+					else
+					{
+						AddDiagnostic(Output.Errors, TEXT("Expected token '{'"));
+					}
+					continue;
+				}
+				else if (StartsWith(Source, TEXTVIEW("}")))
+				{
+					if (NamespaceTracker.Pop())
+					{
+						Source = Source.Mid(1);
+						continue;
+					}
+					else
+					{
+						AddDiagnostic(Output.Errors, TEXT("Expected token '}'"));
+						break;
+					}
+				}
 			}
 		}
-		else if (ChunkType == ECodeChunkType::Operator
+
+		if (ChunkType == ECodeChunkType::Operator
 			&& !PendingBlocks.IsEmpty()
 			&& OperatorKeywordBlockIndex == (PendingBlocks.Num()-1))
 		{
@@ -1074,7 +1275,7 @@ void FindChunksByIdentifier(TConstArrayView<FCodeChunk> Chunks, FStringView Iden
 	{
 		for (const FCodeBlock& Block : Chunk.Blocks)
 		{
-			if (Block.Code == Identifier)
+			if (Block == Identifier)
 			{
 				Callback(Chunk);
 			}
@@ -1113,22 +1314,17 @@ static void ExtractIdentifiers(FStringView InSource, TArray<FStringView>& Result
 {
 	FStringView Source = InSource;
 
-	for (;;)
+	while (!Source.IsEmpty())
 	{
-		Source = SkipSpace(Source);
+		TCHAR FirstChar = Source.GetData()[0];
 
-		if (Source.IsEmpty())
-		{
-			break;
-		}
-
-		if (!IsPossibleIdentifierCharacter(Source[0]))
+		if (FirstChar == L'#')
 		{
-			if (StartsWith(Source, TEXTVIEW("//"))
-				|| StartsWith(Source, TEXTVIEW("#line"))
+			if (StartsWith(Source, TEXTVIEW("#line"))
 				|| StartsWith(Source, TEXTVIEW("#pragma")))
 			{
 				Source = SkipUntilNextLine(Source);
+				Source = SkipSpace(Source);
 				continue;
 			}
 			else if (StartsWith(Source, TEXTVIEW("#if 0")))
@@ -1138,6 +1334,16 @@ static void ExtractIdentifiers(FStringView InSource, TArray<FStringView>& Result
 				{
 					Source = SubStrView(Source, 6);
 				}
+				Source = SkipSpace(Source);
+				continue;
+			}
+		}
+		else if (FirstChar == L'/')
+		{
+			if (StartsWith(Source, TEXTVIEW("//")))
+			{
+				Source = SkipUntilNextLine(Source);
+				Source = SkipSpace(Source);
 				continue;
 			}
 			else if (StartsWith(Source, TEXTVIEW("/*")))
@@ -1147,6 +1353,7 @@ static void ExtractIdentifiers(FStringView InSource, TArray<FStringView>& Result
 				{
 					Source = SubStrView(Source, 2);
 				}
+				Source = SkipSpace(Source);
 				continue;
 			}
 		}
@@ -1170,7 +1377,7 @@ static void ExtractIdentifiers(FStringView InSource, TArray<FStringView>& Result
 			}
 		}
 
-		Source = Remainder;
+		Source = SkipSpace(Remainder);
 	}
 }
 
@@ -1178,11 +1385,11 @@ static void ExtractIdentifiers(const FCodeChunk& Chunk, TArray<FStringView>& Res
 {
 	for (const FCodeBlock& Block : Chunk.Blocks)
 	{
-		ExtractIdentifiers(Block.Code, Result);
+		ExtractIdentifiers(Block, Result);
 	}
 }
 
-static void OutputChunk(const FCodeChunk& Chunk, FStringBuilderBase& OutputStream)
+static void OutputChunk(const FCodeChunk& Chunk, FString& OutputStream)
 {
 	if (Chunk.Blocks.IsEmpty())
 	{
@@ -1192,7 +1399,7 @@ static void OutputChunk(const FCodeChunk& Chunk, FStringBuilderBase& OutputStrea
 	if (Chunk.bVerbatim)
 	{
 		// Fast path to output entire code block verbatim, preserving any new lines and whitespace		
-		OutputStream << Chunk.GetCode();
+		OutputStream.Append(Chunk.GetCode());
 	}
 	else
 	{
@@ -1201,24 +1408,24 @@ static void OutputChunk(const FCodeChunk& Chunk, FStringBuilderBase& OutputStrea
 		{
 			if (Index != 0)
 			{
-				OutputStream << ' ';
+				OutputStream.AppendChar(L' ');
 			}
 
 			if (Block.Type == EBlockType::Expression)
 			{
-				OutputStream << "= ";
+				OutputStream.Append(TEXTVIEW("= "));
 			}
 			else if (Block.Type == EBlockType::Body)
 			{
-				OutputStream << "\n";
+				OutputStream.AppendChar(L'\n');
 			}
 
 			if (Block.Type == EBlockType::Binding || Block.Type == EBlockType::Base)
 			{
-				OutputStream << ": ";
+				OutputStream.Append(TEXTVIEW(": "));
 			}
 
-			OutputStream << Block.Code;
+			OutputStream.Append(Block.GetCode());
 
 			++Index;
 		}
@@ -1231,10 +1438,10 @@ static void OutputChunk(const FCodeChunk& Chunk, FStringBuilderBase& OutputStrea
 		&& Chunk.Type != ECodeChunkType::Define
 		&& Chunk.Type != ECodeChunkType::CommentLine)
 	{
-		OutputStream << ";";
+		OutputStream.AppendChar(L';');
 	}
 
-	OutputStream << "\n";
+	OutputStream.AppendChar(L'\n');
 }
 
 struct FCasedStringViewKeyFuncs : public DefaultKeyFuncs<FStringView>
@@ -1245,43 +1452,52 @@ struct FCasedStringViewKeyFuncs : public DefaultKeyFuncs<FStringView>
 	static FORCEINLINE bool Matches(FStringView A, FStringView B) { return Equals(A, B); }
 	static FORCEINLINE uint32 GetKeyHash(FStringView Key)
 	{
-		return CityHash32((const char*)Key.GetData(), Key.Len() * sizeof(*Key.GetData()));
+		return FXxHash64::HashBuffer(Key.GetData(), Key.Len() * sizeof(*Key.GetData())).Hash;
 	}
 };
 
-static void BuildLineBreakMap(FStringView Source, TArray<int32>& OutLineBreakMap, TArray<FStringView>& OutLineDirectives)
+static void BuildLineBreakMap(FStringView Source, TArray<int32>& OutLineBreakMap)
 {
-	OutLineBreakMap.Empty();
-	OutLineDirectives.Empty();
+	OutLineBreakMap.Reset();
 
-	OutLineBreakMap.Add(0); // Lines numbers are 1-based, so add a dummy element to make LowerBound later return the line number directly
+	OutLineBreakMap.Add(0); // Lines numbers are 1-based, so add a dummy element to make UpperBound later return the line number directly
 
 	const int32 SourceLen = Source.Len();
 	const TCHAR* Chars = Source.GetData(); // avoid bounds check overhead in [] operator
 
-	for (int32 Index = 0; Index < SourceLen; ++Index)
+	int32 Cursor = 0;
+
+#if UE_SHADER_MINIFIER_SSE
+	static_assert(sizeof(*Chars) == 2, "BuildLineBreakMap expects 16 bit characters");
+	const int32 AlignedLen = SourceLen & (~7); // align down to multiple of 8 TCHAR-s
+	const __m128i Needle = _mm_set1_epi16(L'\n');
+	while (Cursor < AlignedLen)
 	{
-		if (Chars[Index] == TCHAR('\n'))
+		__m128i Chunk = _mm_loadu_si128(reinterpret_cast<const __m128i*>(Chars + Cursor));
+		__m128i MaskVec = _mm_cmpeq_epi16(Chunk, Needle);
+		uint32 Mask = _mm_movemask_epi8(MaskVec);
+		while (Mask != 0)
 		{
-			OutLineBreakMap.Add(Index);
+			// NOTE: 2 bits represent each character
+			const uint32 BitIndex = FMath::CountTrailingZeros(Mask);
+			const uint32 ChunkCharIndex = BitIndex / 2;
+			OutLineBreakMap.Add(Cursor + ChunkCharIndex);
+			Mask &= ~(3 << BitIndex);
 		}
-		else if (Chars[Index] == TCHAR('#'))
-		{
-			// In a general case, directives may be inside comments or inactive blocks.
-			// However we expect input source to be fully preprocessed and comments to be removed.
+		Cursor += 8;
+	}
+#endif //UE_SHADER_MINIFIER_SSE
 
-			FStringView PossibleDirective = Source.Mid(Index);
-			if (StartsWith(PossibleDirective, TEXTVIEW("#line")))
-			{
-				FStringView Remainder = SkipUntilNextLine(PossibleDirective);
-				FStringView LineDirective = SubStrView(PossibleDirective, 0, PossibleDirective.Len() - Remainder.Len());
-				OutLineDirectives.Add(LineDirective);
-			}
+	while (Cursor < SourceLen)
+	{
+		if (Chars[Cursor] == TCHAR('\n'))
+		{
+			OutLineBreakMap.Add(Cursor);
 		}
+		++Cursor;
 	}
 }
 
-
 static int32 FindLineDirective(const TArray<FStringView>& LineDirectives, const TCHAR* Ptr)
 {
 	int32 FoundIndex = Algo::UpperBoundBy(LineDirectives, Ptr, [](FStringView Item)
@@ -1347,27 +1563,32 @@ static bool ParseLineDirective(FStringView Input, int32& OutLineNumber, FStringV
 	return true;
 }
 
-static void OpenNamespace(FStringBuilderBase& OutputStream, const FNamespace& Namespace)
+static void OpenNamespace(FString& OutputStream, const FNamespace& Namespace)
 {
 	for (const FStringView& Name : Namespace.Stack)
 	{
-		OutputStream << TEXT("namespace ") << Name << TEXT(" { ");
+		OutputStream.Append(TEXTVIEW("namespace "));
+		OutputStream.Append(Name);
+		OutputStream.Append(TEXTVIEW(" { "));
 	}
 }
 
-static void CloseNamespace(FStringBuilderBase& OutputStream, const FNamespace& Namespace)
+static void CloseNamespace(FString& OutputStream, const FNamespace& Namespace)
 {
 	for (const FStringView& Name : Namespace.Stack)
 	{
-		OutputStream << TEXT("}");
+		OutputStream.AppendChar(L'}');
 	}
 
-	OutputStream << TEXT(" // namespace ") << Namespace.FullName;
+	OutputStream.Append(TEXTVIEW(" // namespace "));
+	OutputStream.Append(Namespace.FullName);
 }
 
 static FString MinifyShader(const FParsedShader& Parsed, TConstArrayView<FStringView> RequiredSymbols, EMinifyShaderFlags Flags, FDiagnostics& Diagnostics)
 {
-	FStringBuilderBase OutputStream;
+	FString OutputStream;
+
+	OutputStream.Reserve(Parsed.Source.Len() / 3); // Heuristic pre-allocation based on average measured reduced code size
 
 	TSet<FStringView, FCasedStringViewKeyFuncs, FDefaultSetAllocator> RelevantIdentifiers;
 
@@ -1562,8 +1783,8 @@ static FString MinifyShader(const FParsedShader& Parsed, TConstArrayView<FString
 
 			if ((Chunk.Type == ECodeChunkType::CBuffer || Chunk.Type == ECodeChunkType::Enum) && Block.Type == EBlockType::Body)
 			{
-				TempIdentifiers.Empty();
-				ExtractIdentifiers(Block.Code, TempIdentifiers);
+				TempIdentifiers.Reset();
+				ExtractIdentifiers(Block, TempIdentifiers);
 				for (FStringView Identifier : TempIdentifiers)
 				{
 					ChunksByIdentifier.FindOrAdd(Identifier).Push(&Chunk);
@@ -1579,7 +1800,7 @@ static FString MinifyShader(const FParsedShader& Parsed, TConstArrayView<FString
 				continue;
 			}
 
-			ChunksByIdentifier.FindOrAdd(Block.Code).Push(&Chunk);
+			ChunksByIdentifier.FindOrAdd(Block).Push(&Chunk);
 		}
 	}
 
@@ -1587,7 +1808,7 @@ static FString MinifyShader(const FParsedShader& Parsed, TConstArrayView<FString
 
 	while (!PendingChunks.IsEmpty())
 	{
-		TempIdentifiers.Empty();
+		TempIdentifiers.Reset();
 
 		const FCodeChunk* CurrentChunk = PendingChunks.Last();
 		PendingChunks.Pop();
@@ -1668,20 +1889,38 @@ static FString MinifyShader(const FParsedShader& Parsed, TConstArrayView<FString
 
 	if (EnumHasAnyFlags(Flags, EMinifyShaderFlags::OutputStats))
 	{
-		OutputStream << "// Total code chunks: " << RelevantChunks.Num() << "\n";
-		OutputStream << "// - Functions: " << NumFunctions << "\n";
-		OutputStream << "// - Structs: " << NumStructs << "\n";
-		OutputStream << "// - CBuffers: " << NumCBuffers << "\n";
-		OutputStream << "// - Variables: " << NumVariables << "\n";
-		OutputStream << "// - Other: " << NumOtherChunks << "\n";
-		OutputStream << "\n";
+		OutputStream.Append(TEXTVIEW("// Total code chunks: "));
+		OutputStream.AppendInt(RelevantChunks.Num());
+		OutputStream.AppendChar(L'\n');
+
+		OutputStream.Append(TEXTVIEW("// - Functions: "));
+		OutputStream.AppendInt(NumFunctions);
+		OutputStream.AppendChar(L'\n');
+
+		OutputStream.Append(TEXTVIEW("// - Structs: "));
+		OutputStream.AppendInt(NumStructs);
+		OutputStream.AppendChar(L'\n');
+
+		OutputStream.Append(TEXTVIEW("// - CBuffers: "));
+		OutputStream.AppendInt(NumCBuffers);
+		OutputStream.AppendChar(L'\n');
+
+		OutputStream.Append(TEXTVIEW("// - Variables: "));
+		OutputStream.AppendInt(NumVariables);
+		OutputStream.AppendChar(L'\n');
+
+		OutputStream.Append(TEXTVIEW("// - Other: "));
+		OutputStream.AppendInt(NumOtherChunks);
+		OutputStream.AppendChar(L'\n');
+
+		OutputStream.AppendChar(L'\n');
 	}
 
 	TArray<int32> LineBreakMap;
-	TArray<FStringView> LineDirectives;
+	const TArray<FStringView>& LineDirectives = Parsed.LineDirectives;
 	if (EnumHasAnyFlags(Flags, EMinifyShaderFlags::OutputLines))
 	{
-		BuildLineBreakMap(Parsed.Source, LineBreakMap, LineDirectives);
+		BuildLineBreakMap(Parsed.Source, LineBreakMap);
 	}
 
 	const FNamespace* CurrentNamespace = nullptr;
@@ -1731,13 +1970,13 @@ static FString MinifyShader(const FParsedShader& Parsed, TConstArrayView<FString
 			if (CurrentNamespace)
 			{
 				CloseNamespace(OutputStream, *CurrentNamespace);
-				OutputStream << TEXT("\n\n");
+				OutputStream.Append(TEXTVIEW("\n\n"));
 			}
 
 			if (PendingNamespace)
 			{
 				OpenNamespace(OutputStream, *PendingNamespace);
-				OutputStream << TEXT("\n\n");
+				OutputStream.Append(TEXTVIEW("\n\n"));
 			}
 
 			CurrentNamespace = PendingNamespace;
@@ -1752,14 +1991,16 @@ static FString MinifyShader(const FParsedShader& Parsed, TConstArrayView<FString
 				FStringView  RequestedByName = RequestedByChunk->FindFirstBlockByType(EBlockType::Name);
 				if (!RequestedByName.IsEmpty())
 				{
-					OutputStream << TEXT("// REASON: ") << RequestedByName << TEXT("\n");
+					OutputStream.Append(TEXTVIEW("// REASON: "));
+					OutputStream.Append(RequestedByName);
+					OutputStream.AppendChar(L'\n');
 				}
 			}
 		}
 
 		if (EnumHasAnyFlags(Flags, EMinifyShaderFlags::OutputLines))
 		{
-			const FStringView ChunkCode = Chunk.Blocks[0].Code;
+			const FStringView ChunkCode = Chunk.Blocks[0];
 			int32 LineDirectiveIndex = FindLineDirective(LineDirectives, ChunkCode.GetData());
 			int32 ChunkLine = FindLineNumber(Parsed.Source, LineBreakMap, ChunkCode.GetData());
 
@@ -1768,7 +2009,9 @@ static FString MinifyShader(const FParsedShader& Parsed, TConstArrayView<FString
 				// There was no valid line directive for this chunk, but we do know the line in the input source, so just emit that.
 				if (ChunkLine > LastLineNumber + 1 || !LastLineFileName.IsEmpty())
 				{
-					OutputStream << TEXT("#line ") << ChunkLine << TEXT("\n");
+					OutputStream.Append(TEXTVIEW("#line "));
+					OutputStream.AppendInt(ChunkLine);
+					OutputStream.AppendChar(L'\n');
 				}
 
 				LastLineNumber = ChunkLine;
@@ -1794,14 +2037,19 @@ static FString MinifyShader(const FParsedShader& Parsed, TConstArrayView<FString
 						// Separate the next block from the previous one when it starts with a line directive
 						if (OutputStream.Len())
 						{
-							OutputStream << "\n";
+							OutputStream.AppendChar(L'\n');
 						}
-						OutputStream << TEXT("#line ") << PatchedLineNumber;
+
+						OutputStream.Append(TEXTVIEW("#line "));
+						OutputStream.AppendInt(PatchedLineNumber);
+
 						if (!ParsedFileName.IsEmpty())
 						{
-							OutputStream << TEXT(" \"") << ParsedFileName << TEXT("\"");
+							OutputStream.Append(TEXTVIEW(" \""));
+							OutputStream.Append(ParsedFileName);
+							OutputStream.AppendChar(L'\"');
 						}
-						OutputStream << TEXT("\n");
+						OutputStream.AppendChar(L'\n');
 					}
 
 					LastLineNumber = PatchedLineNumber;
@@ -1816,13 +2064,11 @@ static FString MinifyShader(const FParsedShader& Parsed, TConstArrayView<FString
 	if (CurrentNamespace)
 	{
 		CloseNamespace(OutputStream, *CurrentNamespace);
-		OutputStream << TEXT("\n");
+		OutputStream.AppendChar(L'\n');
 		CurrentNamespace = nullptr;
 	}
 
-	FString Output = FString(OutputStream.ToView());
-
-	return Output;
+	return OutputStream;
 }
 
 static FString MinifyShader(const FParsedShader& Parsed, FStringView EntryPoint, EMinifyShaderFlags Flags, FDiagnostics& Diagnostics)
@@ -2330,7 +2576,7 @@ void MainCS()
 		{
 			for (const FCodeBlock& Block : Chunk.Blocks)
 			{
-				if (Block.Code == Name)
+				if (Block == Name)
 				{
 					return true;
 				}
