Fix compile error on linux/mac due to indexing arrays with char type

diff --git a/Engine/Source/Developer/ShaderCompilerCommon/Private/ShaderMinifier.cpp b/Engine/Source/Developer/ShaderCompilerCommon/Private/ShaderMinifier.cpp
--- a/Engine/Source/Developer/ShaderCompilerCommon/Private/ShaderMinifier.cpp
+++ b/Engine/Source/Developer/ShaderCompilerCommon/Private/ShaderMinifier.cpp
@@ -101,35 +101,35 @@ struct FCharacterFlags
 	{
 		for (char C = '0'; C <= '9'; ++C)
 		{
-			Flags[C] |= Number;
+			Flags[uint8(C)] |= Number;
 		}
 
 		for (char C = 'a'; C <= 'z'; ++C)
 		{
-			Flags[C] |= Letter;
+			Flags[uint8(C)] |= Letter;
 		}
 
 		for (char C = 'A'; C <= 'Z'; ++C)
 		{
-			Flags[C] |= Letter;
+			Flags[uint8(C)] |= Letter;
 		}
 
-		Flags['_'] |= Underscore;
+		Flags[uint8('_')] |= Underscore;
 
-		Flags[' ']  |= Space;
-		Flags['\f'] |= Space;
-		Flags['\r'] |= Space;
-		Flags['\n'] |= Space;
-		Flags['\t'] |= Space;
-		Flags['\v'] |= Space;
+		Flags[uint8(' ')]  |= Space;
+		Flags[uint8('\f')] |= Space;
+		Flags[uint8('\r')] |= Space;
+		Flags[uint8('\n')] |= Space;
+		Flags[uint8('\t')] |= Space;
+		Flags[uint8('\v')] |= Space;
 
-		Flags['\\'] |= Special;
-		Flags['/']  |= Special;
-		Flags['#']  |= Special;
-		Flags['{']  |= Special;
-		Flags['}']  |= Special;
-		Flags['(']  |= Special;
-		Flags[')']  |= Special;
+		Flags[uint8('\\')] |= Special;
+		Flags[uint8('/')]  |= Special;
+		Flags[uint8('#')]  |= Special;
+		Flags[uint8('{')]  |= Special;
+		Flags[uint8('}')]  |= Special;
+		Flags[uint8('(')]  |= Special;
+		Flags[uint8(')')]  |= Special;
 	}
 
 	bool IsSpace(TCHAR C) const
