Add CPU profiler scope to UE::ShaderCompilerCommon::RemoveDeadCode

diff --git a/Engine/Source/Developer/ShaderCompilerCommon/Private/ShaderCompilerCommon.cpp b/Engine/Source/Developer/ShaderCompilerCommon/Private/ShaderCompilerCommon.cpp
--- a/Engine/Source/Developer/ShaderCompilerCommon/Private/ShaderCompilerCommon.cpp
+++ b/Engine/Source/Developer/ShaderCompilerCommon/Private/ShaderCompilerCommon.cpp
@@ -13,6 +13,7 @@
 #include "ShaderSymbolExport.h"
 #include "ShaderMinifier.h"
 #include "Algo/Sort.h"
+#include "ProfilingDebugging/CpuProfilerTrace.h"
 
 static TAutoConsoleVariable<bool> CVarShaderCompilerCleanupUniformBufferCodeNew(
 	TEXT("r.ShaderCompiler.CleanupUniformBufferCodeNew"),
@@ -515,6 +516,8 @@ void UE::ShaderCompilerCommon::ParseRayTracingEntryPoint(const FString& Input, F
 
 bool UE::ShaderCompilerCommon::RemoveDeadCode(FString& InOutPreprocessedShaderSource, TConstArrayView<FStringView> RequiredSymbols, TArray<FShaderCompilerError>& OutErrors)
 {
+	TRACE_CPUPROFILER_EVENT_SCOPE(RemoveDeadCode);
+
 	UE::ShaderMinifier::EMinifyShaderFlags ExtraFlags = UE::ShaderMinifier::EMinifyShaderFlags::None;
 
 #if 0 // Extra features that may be useful during development / debugging
