Extend HlslParser to support Substrate:
- Support struct forward declarations.
- Support static member functions.
- Support scoped function definitions.

#rnx

diff --git a/Engine/Shaders/Public/ShaderVersion.ush b/Engine/Shaders/Public/ShaderVersion.ush
--- a/Engine/Shaders/Public/ShaderVersion.ush
+++ b/Engine/Shaders/Public/ShaderVersion.ush
@@ -2,4 +2,4 @@
 // Each time the console command is executed it generates a new GUID. As this file is included
 // in Platform.ush (which should be included in any shader) it allows to invalidate the shader DDC.
 // 
-#pragma message("UESHADERMETADATA_VERSION 7720193D-8154-49D2-973E-86527A291B62")
+#pragma message("UESHADERMETADATA_VERSION 932D8FB4-8E40-48BC-AB6B-E89472D1644C")
diff --git a/Engine/Source/Developer/ShaderCompilerCommon/Private/HlslAST.cpp b/Engine/Source/Developer/ShaderCompilerCommon/Private/HlslAST.cpp
--- a/Engine/Source/Developer/ShaderCompilerCommon/Private/HlslAST.cpp
+++ b/Engine/Source/Developer/ShaderCompilerCommon/Private/HlslAST.cpp
@@ -564,9 +564,12 @@ namespace CrossCompiler
 		FFunction::FFunction(FLinearAllocator* InAllocator, const FSourceInfo& InInfo) :
 			FNode(InAllocator, InInfo),
 			ReturnType(nullptr),
+			ScopeIdentifier(nullptr),
 			Identifier(nullptr),
 			ReturnSemantic(nullptr),
-			Parameters(InAllocator)
+			Parameters(InAllocator),
+			bIsDefinition(false),
+			bIsStatic(false)
 		{
 		}
 
@@ -574,8 +577,16 @@ namespace CrossCompiler
 		{
 			WriteAttributes(Writer);
 			Writer << TEXT("\n");
+			if (bIsStatic)
+			{
+				Writer << TEXT("static ");
+			}
 			ReturnType->Write(Writer);
 			Writer << (TCHAR)' ';
+			if (ScopeIdentifier)
+			{
+				Writer << ScopeIdentifier << TEXT("::");
+			}
 			Writer << Identifier;
 			Writer << (TCHAR)'(';
 			bool bFirst = true;
@@ -1458,7 +1469,8 @@ namespace CrossCompiler
 			FNode(InAllocator, InInfo),
 			Name(nullptr),
 			ParentName(nullptr),
-			Members(InAllocator)
+			Members(InAllocator),
+			bForwardDeclaration(false)
 		{
 		}
 
@@ -1471,18 +1483,25 @@ namespace CrossCompiler
 				Writer << TEXT(" : ");
 				Writer << ParentName;
 			}
-			Writer << (TCHAR)'\n';
-			Writer.DoIndent();
-			Writer << TEXT("{\n");
-
-			for (auto* Member : Members)
+			if (bForwardDeclaration)
 			{
-				FASTWriterIncrementScope Scope(Writer);
-				Member->Write(Writer);
+				Writer << TEXT(";\n");
 			}
+			else
+			{
+				Writer << TEXT('\n');
+				Writer.DoIndent();
+				Writer << TEXT("{\n");
 
-			Writer.DoIndent();
-			Writer << TEXT("}");
+				for (auto* Member : Members)
+				{
+					FASTWriterIncrementScope Scope(Writer);
+					Member->Write(Writer);
+				}
+
+				Writer.DoIndent();
+				Writer << TEXT("}");
+			}
 		}
 
 		FStructSpecifier::~FStructSpecifier()
diff --git a/Engine/Source/Developer/ShaderCompilerCommon/Private/HlslAST.h b/Engine/Source/Developer/ShaderCompilerCommon/Private/HlslAST.h
--- a/Engine/Source/Developer/ShaderCompilerCommon/Private/HlslAST.h
+++ b/Engine/Source/Developer/ShaderCompilerCommon/Private/HlslAST.h
@@ -555,6 +555,7 @@ namespace CrossCompiler
 			const TCHAR* Name;
 			const TCHAR* ParentName;
 			TLinearArray<FNode*> Members;
+			bool bForwardDeclaration;
 		};
 
 		struct FCBufferDeclaration : public FNode
@@ -644,12 +645,14 @@ namespace CrossCompiler
 			virtual void Write(FASTWriter& Writer) const override;
 
 			FFullySpecifiedType* ReturnType;
+			const TCHAR* ScopeIdentifier;
 			const TCHAR* Identifier;
 			FSemanticSpecifier* ReturnSemantic;
 
 			TLinearArray<FNode*> Parameters;
 
-			bool bIsDefinition = false;
+			bool bIsDefinition;
+			bool bIsStatic;
 
 			//Signature
 		};
diff --git a/Engine/Source/Developer/ShaderCompilerCommon/Private/HlslParser.cpp b/Engine/Source/Developer/ShaderCompilerCommon/Private/HlslParser.cpp
--- a/Engine/Source/Developer/ShaderCompilerCommon/Private/HlslParser.cpp
+++ b/Engine/Source/Developer/ShaderCompilerCommon/Private/HlslParser.cpp
@@ -11,7 +11,7 @@
 namespace CrossCompiler
 {
 	EParseResult ParseExpressionStatement(class FHlslParser& Parser, FLinearAllocator* Allocator, AST::FNode** OutStatement);
-	EParseResult ParseStructBody(FHlslParser& Parser, FSymbolScope* SymbolScope, FLinearAllocator* Allocator, AST::FTypeSpecifier** OutTypeSpecifier);
+	EParseResult ParseStructDeclaration(FHlslParser& Parser, FSymbolScope* SymbolScope, FLinearAllocator* Allocator, AST::FTypeSpecifier** OutTypeSpecifier);
 	EParseResult TryParseAttribute(FHlslParser& Parser, FLinearAllocator* Allocator, AST::FAttribute** OutAttribute);
 	EParseResult ParseFunctionDeclaration(FHlslParser& Parser, FLinearAllocator* Allocator, TLinearArray<AST::FAttribute*>& Attributes, AST::FNode** OutFunction);
 
@@ -750,7 +750,7 @@ namespace CrossCompiler
 
 		if (Parser.Scanner.MatchToken(EHlslToken::Struct))
 		{
-			auto Result = ParseStructBody(Parser, SymbolScope, Allocator, &FullType->Specifier);
+			auto Result = ParseStructDeclaration(Parser, SymbolScope, Allocator, &FullType->Specifier);
 			if (Result != EParseResult::Matched)
 			{
 				return ParseResultError();
@@ -851,7 +851,7 @@ namespace CrossCompiler
 			}
 			else if (Result == EParseResult::Error)
 			{
-				return EParseResult::Error;
+				return ParseResultError();
 			}
 			else if (bCanBeUnmatched && Result == EParseResult::NotMatched)
 			{
@@ -932,44 +932,11 @@ namespace CrossCompiler
 		return ParseResultError();
 	}
 
-	EParseResult ParseStructBody(FHlslParser& Parser, FSymbolScope* SymbolScope, FLinearAllocator* Allocator, AST::FTypeSpecifier** OutTypeSpecifier)
+	static EParseResult ParseStructBody(FHlslParser& Parser, FSymbolScope* SymbolScope, FLinearAllocator* Allocator, AST::FStructSpecifier* Struct)
 	{
-		const auto* Name = Parser.Scanner.GetCurrentToken();
-		if (!Name)
-		{
-			return ParseResultError();
-		}
-
-		bool bAnonymous = true;
-		if (Parser.Scanner.MatchToken(EHlslToken::Identifier))
-		{
-			bAnonymous = false;
-			SymbolScope->Add(Name->String);
-		}
-
-		const TCHAR* Parent = nullptr;
-		if (Parser.Scanner.MatchToken(EHlslToken::Colon))
-		{
-			const auto* ParentToken = Parser.Scanner.GetCurrentToken();
-			if (!Parser.Scanner.MatchToken(EHlslToken::Identifier))
-			{
-				Parser.Scanner.SourceError(TEXT("Identifier expected!\n"));
-				return ParseResultError();
-			}
-
-			Parent = Allocator->Strdup(ParentToken->String);
-		}
-
-		if (!Parser.Scanner.MatchToken(EHlslToken::LeftBrace))
-		{
-			Parser.Scanner.SourceError(TEXT("Expected '{'!"));
-			return ParseResultError();
-		}
-
-		auto* Struct = new(Allocator) AST::FStructSpecifier(Allocator, Name->SourceInfo);
-		Struct->ParentName = Allocator->Strdup(Parent);
-		//@todo-rco: Differentiate anonymous!
-		Struct->Name = bAnonymous ? nullptr : Allocator->Strdup(Name->String);
+		check(SymbolScope != nullptr);
+		check(Allocator != nullptr);
+		check(Struct != nullptr);
 
 		bool bFoundRightBrace = false;
 		while (Parser.Scanner.HasMoreTokens())
@@ -1043,6 +1010,59 @@ namespace CrossCompiler
 			return ParseResultError();
 		}
 
+		return EParseResult::Matched;
+	}
+		
+	EParseResult ParseStructDeclaration(FHlslParser& Parser, FSymbolScope* SymbolScope, FLinearAllocator* Allocator, AST::FTypeSpecifier** OutTypeSpecifier)
+	{
+		const auto* Name = Parser.Scanner.GetCurrentToken();
+		if (!Name)
+		{
+			return ParseResultError();
+		}
+
+		bool bAnonymous = true;
+		if (Parser.Scanner.MatchToken(EHlslToken::Identifier))
+		{
+			bAnonymous = false;
+			SymbolScope->Add(Name->String);
+		}
+
+		const TCHAR* Parent = nullptr;
+		if (Parser.Scanner.MatchToken(EHlslToken::Colon))
+		{
+			const auto* ParentToken = Parser.Scanner.GetCurrentToken();
+			if (!Parser.Scanner.MatchToken(EHlslToken::Identifier))
+			{
+				Parser.Scanner.SourceError(TEXT("Identifier expected!\n"));
+				return ParseResultError();
+			}
+
+			Parent = Allocator->Strdup(ParentToken->String);
+		}
+
+		const bool bForwardDeclaration = !Parser.Scanner.MatchToken(EHlslToken::LeftBrace);
+		if (bForwardDeclaration && bAnonymous)
+		{
+			Parser.Scanner.SourceError(TEXT("Anonymous struct must have a body definition"));
+			return ParseResultError();
+		}
+
+		auto* Struct = new(Allocator) AST::FStructSpecifier(Allocator, Name->SourceInfo);
+		Struct->ParentName = Allocator->Strdup(Parent);
+		//@todo-rco: Differentiate anonymous!
+		Struct->Name = bAnonymous ? nullptr : Allocator->Strdup(Name->String);
+		Struct->bForwardDeclaration = bForwardDeclaration;
+
+		if (!bForwardDeclaration)
+		{
+			const EParseResult Result = ParseStructBody(Parser, SymbolScope, Allocator, Struct);
+			if (Result != EParseResult::Matched)
+			{
+				return Result;
+			}
+		}
+
 		auto* TypeSpecifier = new(Allocator) AST::FTypeSpecifier(Allocator, Struct->SourceInfo);
 		TypeSpecifier->Structure = Struct;
 		*OutTypeSpecifier = TypeSpecifier;
@@ -1112,7 +1132,7 @@ namespace CrossCompiler
 
 		check(Result == EParseResult::Matched);
 
-		auto* Identifier = Parser.Scanner.GetCurrentToken();
+		const FHlslToken* Identifier = Parser.Scanner.GetCurrentToken();
 		if (!Parser.Scanner.MatchToken(EHlslToken::Identifier))
 		{
 			// This could be an error... But we should allow testing for a global variable before any rash decisions
@@ -1120,6 +1140,20 @@ namespace CrossCompiler
 			return EParseResult::NotMatched;
 		}
 
+		// Check for scoped function declarations, e.g. "MyStruct::MyFunction() ..."
+		const FHlslToken* ScopeIdentifier = nullptr;
+		if (Parser.Scanner.MatchToken(EHlslToken::ColonColon))
+		{
+			ScopeIdentifier = Identifier;
+			Identifier = Parser.Scanner.GetCurrentToken();
+			if (!Parser.Scanner.MatchToken(EHlslToken::Identifier))
+			{
+				// This could be an error... But we should allow testing for a global variable before any rash decisions
+				Parser.Scanner.SetCurrentTokenIndex(OriginalToken);
+				return EParseResult::NotMatched;
+			}
+		}
+
 		if (!Parser.Scanner.MatchToken(EHlslToken::LeftParenthesis))
 		{
 			// This could be an error... But we should allow testing for a global variable before any rash decisions
@@ -1131,6 +1165,10 @@ namespace CrossCompiler
 
 		auto* Function = new(Allocator) AST::FFunction(Allocator, Identifier->SourceInfo);
 		Function->Identifier = Allocator->Strdup(Identifier->String);
+		if (ScopeIdentifier)
+		{
+			Function->ScopeIdentifier = Allocator->Strdup(ScopeIdentifier->String);
+		}
 		Function->ReturnType = new(Allocator) AST::FFullySpecifiedType(Allocator, TypeSpecifier->SourceInfo);
 		Function->ReturnType->Specifier = TypeSpecifier;
 
@@ -1219,12 +1257,48 @@ Done:
 		return ParseResultError();
 	}
 
+	bool ParseFunctionStorageSpecifiers(FHlslParser& Parser, bool& bOutInline, bool& bOutStatic)
+	{
+		bOutInline = false;
+		bOutStatic = false;
+		while (true)
+		{
+			if (Parser.Scanner.MatchToken(EHlslToken::Inline))
+			{
+				if (bOutInline)
+				{
+					Parser.Scanner.SourceError(TEXT("Duplicate storage qualifier found: 'inline'"));
+					return false;
+				}
+				bOutInline = true;
+			}
+			else if (Parser.Scanner.MatchToken(EHlslToken::Static))
+			{
+				if (bOutStatic)
+				{
+					Parser.Scanner.SourceError(TEXT("Duplicate storage qualifier found: 'static'"));
+					return false;
+				}
+				bOutStatic = true;
+			}
+			else
+			{
+				break;
+			}
+		}
+		return true;
+	}
+
 	EParseResult ParseFunctionDeclaration(FHlslParser& Parser, FLinearAllocator* Allocator, TLinearArray<AST::FAttribute*>& Attributes, AST::FNode** OutFunction)
 	{
 		const auto* CurrentToken = Parser.Scanner.GetCurrentToken();
 
 		// Inline could be used but will be ignored per the hlsl spec. If found then this HAS to be a function.
-		bool bFoundInline = Parser.Scanner.MatchToken(EHlslToken::Inline);
+		bool bFoundInline = false, bFoundStatic = false;
+		if (!ParseFunctionStorageSpecifiers(Parser, bFoundInline, bFoundStatic))
+		{
+			return ParseResultError();
+		}
 
 		AST::FFunction* Function = nullptr;
 		EParseResult Result = ParseFunctionDeclarator(Parser, Allocator, &Function);
@@ -1242,6 +1316,7 @@ Done:
 		{
 			return Result;
 		}
+		Function->bIsStatic = bFoundStatic;
 
 		if (Parser.Scanner.MatchToken(EHlslToken::Semicolon))
 		{
