Shader Preprocessor:  Compact commented out uniform buffer declarations out of the source code in CleanupUniformBufferCode.  Saves 9% of overall preprocess time by avoiding comment scanning in minifier parse step.  Also took this opportunity to remove the unused original CleanupUniformBufferCode implementation -- the new version has been stable, so we don't need a CVar toggle any more.

#rnx

diff --git a/Engine/Source/Developer/ShaderCompilerCommon/Private/ShaderCompilerCommon.cpp b/Engine/Source/Developer/ShaderCompilerCommon/Private/ShaderCompilerCommon.cpp
--- a/Engine/Source/Developer/ShaderCompilerCommon/Private/ShaderCompilerCommon.cpp
+++ b/Engine/Source/Developer/ShaderCompilerCommon/Private/ShaderCompilerCommon.cpp
@@ -15,13 +15,6 @@
 #include "Algo/Sort.h"
 #include "ProfilingDebugging/CpuProfilerTrace.h"
 
-static TAutoConsoleVariable<bool> CVarShaderCompilerCleanupUniformBufferCodeNew(
-	TEXT("r.ShaderCompiler.CleanupUniformBufferCodeNew"),
-	true,
-	TEXT("Run new optimized version of CleanupUniformBufferCode.  Temporary chicken switch in case there are issues, before we remove old version."),
-	ECVF_Default
-);
-
 IMPLEMENT_MODULE(FDefaultModuleImpl, ShaderCompilerCommon);
 
 int16 GetNumUniformBuffersUsed(const FShaderCompilerResourceTable& InSRT)
@@ -380,37 +373,6 @@ const TCHAR* ParseHLSLSymbolName(const TCHAR* SearchString, FString& SymbolName)
 	return Result.GetData() + Result.Len();
 }
 
-const TCHAR* MatchStructMemberName(const FString& SymbolName, const TCHAR* SearchPtr)
-{
-	for (int32 i = 0; i < SymbolName.Len(); i++)
-	{
-		if (*SearchPtr != SymbolName[i])
-		{
-			return nullptr;
-		}
-		
-		SearchPtr++;
-
-		if (i < SymbolName.Len() - 1)
-		{
-			// Skip whitespace within the struct member reference before the end
-			// eg 'View. ViewToClip'
-			while (IsSpaceOrTabOrEOL(*SearchPtr))
-			{
-				SearchPtr++;
-			}
-		}
-	}
-
-	// Only match whole symbol
-	if (IsValidHLSLIdentifierCharacter(*SearchPtr))
-	{
-		return nullptr;
-	}
-
-	return SearchPtr;
-}
-
 EShaderParameterType UE::ShaderCompilerCommon::ParseParameterType(
 	FStringView InType,
 	TArrayView<const TCHAR* const> InExtraSRVTypes,
@@ -852,105 +814,12 @@ struct FUniformBufferInfoNew
 	TArray<FUniformBufferMemberView> MembersByLength;	// Offset and count of members of a given length
 };
 
-static void FillUniformBufferDefinition(TStringBuilder<128>& Builder, const TCHAR* Start, const TCHAR* End)
+// Tracks the offset and length of commented out uniform buffer declarations in the source code, so we can compact them out
+struct FUniformBufferSpan
 {
-	Builder.Reset();
-
-	const TCHAR* Cursor = Start;
-	
-	while (Cursor != End)
-	{
-		if (!IsSpaceOrTabOrEOL(*Cursor))
-		{
-			Builder.AppendChar(*Cursor);
-		}
-
-		Cursor++;
-	}
-}
-
-const TCHAR* ParseUniformBufferDefinition(const TCHAR* SourceStart, const TCHAR* ReadStart, TMap<FStringView, FUniformBufferInfo>& UniformBufferInfos)
-{
-	const FStringView UniformBufferName(ParseHLSLSymbolName(ReadStart));
-
-	const TCHAR* OpeningBrace = FindNextChar(ReadStart, '{');
-	const TCHAR* ClosingBrace = FindMatchingClosingBrace(OpeningBrace + 1);
-
-	const TCHAR* CurrentParseStart = OpeningBrace + 1;
-	const TCHAR* NextSemicolon = FindNextChar(CurrentParseStart, ';');
-
-	TStringBuilder<128> StructName;
-	TStringBuilder<128> GlobalName;
-	
-	while (NextSemicolon < ClosingBrace)
-	{
-		const TCHAR* NextSeparator = FindNextChar(CurrentParseStart, '=');
-		if (NextSeparator < NextSemicolon)
-		{
-			const TCHAR* StructStart = CurrentParseStart;
-			const TCHAR* StructEnd = NextSeparator - 1;
-
-			const TCHAR* GlobalStart = NextSeparator + 1;
-			const TCHAR* GlobalEnd = NextSemicolon - 1;
-
-			FillUniformBufferDefinition(StructName, StructStart, StructEnd);
-			FillUniformBufferDefinition(GlobalName, GlobalStart, GlobalEnd);
-
-			// Avoid unnecessary conversions
-			if (StructName.Len() == GlobalName.Len() && FCString::Strncmp(*StructName, *GlobalName, StructName.Len()) != 0)
-			{
-				FUniformBufferMemberInfo NewMemberInfo;
-				NewMemberInfo.NameAsStructMember = StructName;
-				NewMemberInfo.GlobalName = GlobalName;
-
-				// Add this member to the map
-				if (FUniformBufferInfo* Info = UniformBufferInfos.Find(UniformBufferName))
-				{
-					Info->Members.Add(MoveTemp(NewMemberInfo));
-				}
-				else
-				{
-					FUniformBufferInfo NewInfo;
-					NewInfo.DefinitionEndOffset = ClosingBrace - SourceStart;
-					NewInfo.Members.Add(MoveTemp(NewMemberInfo));
-
-					UniformBufferInfos.Emplace(UniformBufferName, MoveTemp(NewInfo));
-				}
-			}
-		}
-
-		CurrentParseStart = NextSemicolon + 1;
-		NextSemicolon = FindNextChar(CurrentParseStart, ';');
-	}
-
-	const TCHAR* EndPtr = ClosingBrace;
-
-	// Skip to the end of the UniformBuffer
-	while (*EndPtr && *EndPtr != ';')
-	{
-		EndPtr++;
-	}
-
-	return EndPtr;
-}
-
-inline bool IsValidSymbolStart(const TCHAR* SearchPtr, int32 UniformBufferNameLen)
-{
-	// Make sure this isn't the end of another symbol or a sub-member of a different struct
-	//   ex: "OpaqueBasePass.Substrate.MaxBytesPerPixel" when matching "Substrate"
-	if (IsValidHLSLIdentifierCharacter(*(SearchPtr - 1)) || *(SearchPtr - 1) == '.')
-	{
-		return false;
-	}
-
-	// Avoid lines like '#line   1 "/Engine/Generated/UniformBuffers/View.ush"'
-	if (FCString::Strncmp(SearchPtr + UniformBufferNameLen, TEXT("ush"), 3) == 0)
-	{
-		return false;
-	}
-
-	return true;
-}
+	int32 Offset;
+	int32 Length;
+};
 
 // Compacts spaces out of a compound identifier.  Returns the new end pointer of the compacted identifier.
 // End and result pointers are exclusive (length of the string is End - Start).
@@ -986,7 +855,7 @@ static TCHAR* CompactCompoundIdentifier(TCHAR* Start, TCHAR* End)
 	return WriteChar;
 }
 
-const TCHAR* ParseUniformBufferDefinitionNew(const TCHAR* ReadStart, TArray<FUniformBufferInfoNew>& UniformBufferInfos, uint64 UniformBufferFilter[64], int32 UniformBuffersByLength[64])
+const TCHAR* ParseUniformBufferDefinition(const TCHAR* ReadStart, TArray<FUniformBufferInfoNew>& UniformBufferInfos, uint64 UniformBufferFilter[64], int32 UniformBuffersByLength[64])
 {
 	// TODO:  should we check for an existing item?  In my testing, there's only one uniform buffer declaration with a given name,
 	// but the original code used a map, theoretically allowing for multiple.
@@ -1369,14 +1238,19 @@ const TCHAR* FindPreviousDot(const TCHAR* SearchPtr, const TCHAR* SearchMin)
 	return SearchPtr;
 }
 
-// Unlike the original version, this parses uniform buffers and replaces symbols in a single pass.
-void CleanupUniformBufferCodeNew(const FShaderCompilerEnvironment& Environment, FString& PreprocessedShaderSource)
+// The cross compiler doesn't yet support struct initializers needed to construct static structs for uniform buffers
+// Replace all uniform buffer struct member references (View.WorldToClip) with a flattened name that removes the struct dependency (View_WorldToClip)
+void CleanupUniformBufferCode(const FShaderCompilerEnvironment& Environment, FString& PreprocessedShaderSource)
 {
+	TRACE_CPUPROFILER_EVENT_SCOPE(CleanupUniformBufferCode);
+
 	TArray<FUniformBufferInfoNew> UniformBufferInfos;
+	TArray<FUniformBufferSpan> UniformBufferSpans;
 	uint64 UniformBufferFilter[64] = { 0 };			// A bit set for valid start characters for uniform buffer name of given length
 	int32 UniformBuffersByLength[64];				// Linked list head index into UniformBufferInfos by length (connected by "NextWithSameLength")
 
 	UniformBufferInfos.Reserve(Environment.UniformBufferMap.Num());
+	UniformBufferSpans.Reserve(Environment.UniformBufferMap.Num());
 	memset(UniformBuffersByLength, 0xff, sizeof(UniformBuffersByLength));
 
 	const TCHAR* UniformBufferStructIdentifier = TEXT("UniformBuffer");
@@ -1476,7 +1350,7 @@ void CleanupUniformBufferCodeNew(const FShaderCompilerEnvironment& Environment,
 			// Unterminate the string (put the first character of the struct identifier back in place) and parse it
 			*SearchPtr = UniformBufferStructIdentifier[0];
 
-			const TCHAR* ConstStructEndPtr = ParseUniformBufferDefinitionNew(SearchPtr + StructIdentifierLen, UniformBufferInfos, UniformBufferFilter, UniformBuffersByLength);
+			const TCHAR* ConstStructEndPtr = ParseUniformBufferDefinition(SearchPtr + StructIdentifierLen, UniformBufferInfos, UniformBufferFilter, UniformBuffersByLength);
 			TCHAR* StructEndPtr = &PreprocessedShaderSource[ConstStructEndPtr - &PreprocessedShaderSource[0]];
 
 			// Comment out the uniform buffer struct and initializer
@@ -1485,128 +1359,33 @@ void CleanupUniformBufferCodeNew(const FShaderCompilerEnvironment& Environment,
 			*(StructEndPtr - 1) = '*';
 			*StructEndPtr = '/';
 
+			UniformBufferSpans.Add({ (int32)(SearchPtr - SourceStart), (int32)(StructEndPtr + 1 - SearchPtr) });
+
 			EndOfPreviousUniformBuffer = StructEndPtr + 1;
 			SearchPtr = StructEndPtr + 1;
 		}
 
 	} while (bUniformBufferFound);
-}
-
-// The cross compiler doesn't yet support struct initializers needed to construct static structs for uniform buffers
-// Replace all uniform buffer struct member references (View.WorldToClip) with a flattened name that removes the struct dependency (View_WorldToClip)
-void CleanupUniformBufferCode(const FShaderCompilerEnvironment& Environment, FString& PreprocessedShaderSource)
-{
-	TRACE_CPUPROFILER_EVENT_SCOPE(CleanupUniformBufferCode);
-
-	// Putting new code on a chicken switch in case there are issues
-	static const bool bRunCleanupUniformBufferCodeNew = CVarShaderCompilerCleanupUniformBufferCodeNew.GetValueOnAnyThread();
-	if (bRunCleanupUniformBufferCodeNew)
-	{
-		CleanupUniformBufferCodeNew(Environment, PreprocessedShaderSource);
-		return;
-	}
-
-	TMap<FStringView, FUniformBufferInfo> UniformBufferInfos;
-	UniformBufferInfos.Reserve(Environment.UniformBufferMap.Num());
-
-	// Build a mapping from uniform buffer name to its members
-	{
-		const TCHAR* UniformBufferStructIdentifier = TEXT("UniformBuffer");
-		const int32 StructIdentifierLen = FCString::Strlen(UniformBufferStructIdentifier);
-
-		const TCHAR* SourceStart = &PreprocessedShaderSource[0];
-		TCHAR* SearchPtr = FCString::Strstr(&PreprocessedShaderSource[0], UniformBufferStructIdentifier);
-
-		while (SearchPtr)
-		{
-			const ptrdiff_t Offset = SearchPtr - SourceStart;
-			if (Offset > 0 && IsSpaceOrTabOrEOL(*(SearchPtr - 1)) && IsSpaceOrTabOrEOL(*(SearchPtr + StructIdentifierLen)))
-			{
-				const TCHAR* ConstStructEndPtr = ParseUniformBufferDefinition(&PreprocessedShaderSource[0], SearchPtr + StructIdentifierLen, UniformBufferInfos);
-				TCHAR* StructEndPtr = &PreprocessedShaderSource[ConstStructEndPtr - &PreprocessedShaderSource[0]];
 
-				// Comment out the uniform buffer struct and initializer
-				*SearchPtr = '/';
-				*(SearchPtr + 1) = '*';
-				*(StructEndPtr - 1) = '*';
-				*StructEndPtr = '/';
+	// Compact commented out uniform buffers out of the output source.  This costs around 10x less to do here than later in the minifier.  Note that
+	// it's not necessary to add a line directive to fix up line numbers because uniform buffer declarations are always in generated files, and there
+	// will be a line directive already there for the transition from the generated file back to whatever file included it.  The destination offset
+	// for the first move is the start of the first uniform buffer declaration we are overwriting, then advances as characters are copied.
+	int32 DestOffset = UniformBufferSpans.Num() ? UniformBufferSpans[0].Offset : PreprocessedShaderSource.GetCharArray().Num();
 
-				SearchPtr = FCString::Strstr(StructEndPtr, UniformBufferStructIdentifier);
-			}
-			else
-			{
-				SearchPtr = FCString::Strstr(SearchPtr + 1, UniformBufferStructIdentifier);
-			}
-		}
-	}
-
-	const TCHAR* const SearchSuffixes[] = { TEXT("."), TEXT(" .") };
-	TStringBuilder<64> UniformBufferName;
-
-	// Replace all uniform buffer struct member references (View.WorldToClip) with a flattened name that removes the struct dependency (View_WorldToClip)
-	for (TMap<FStringView, FUniformBufferInfo>::TConstIterator It(UniformBufferInfos); It; ++It)
+	for (int32 SpanIndex = 0; SpanIndex < UniformBufferSpans.Num(); SpanIndex++)
 	{
-		const FUniformBufferInfo& UniformBufferInfo = It.Value();
+		// The source code we are compacting down is from the end of one span to the start of the next span, or end of the string (plus one so it also
+		// moves the null terminator).
+		int32 SourceOffset = UniformBufferSpans[SpanIndex].Offset + UniformBufferSpans[SpanIndex].Length;
+		int32 MoveCount = (SpanIndex < UniformBufferSpans.Num() - 1 ? UniformBufferSpans[SpanIndex + 1].Offset : PreprocessedShaderSource.Len() + 1) - SourceOffset;
 
-		// Start after the definition of this uniform buffer
-		TCHAR* const SearchStart = &PreprocessedShaderSource[UniformBufferInfo.DefinitionEndOffset];
-
-		// Trying "." and " ." separately can be faster than checking for each case as we iterate
-		for (const TCHAR* SearchSuffix : SearchSuffixes)
-		{
-			UniformBufferName.Reset();
-			UniformBufferName.Append(It.Key());
-			UniformBufferName.Append(SearchSuffix);
+		check(DestOffset >= 0 && DestOffset < SourceOffset && SourceOffset + MoveCount <= PreprocessedShaderSource.GetCharArray().Num());
 
-			// Search for the uniform buffer name first, as an optimization (instead of searching the entire source for every member)
-			TCHAR* SearchPtr = FCString::Strstr(SearchStart, *UniformBufferName);
-
-			while (SearchPtr)
-			{
-				// Only match whole symbol
-				if (IsValidSymbolStart(SearchPtr, UniformBufferName.Len()))
-				{
-					// Find the matching member we are replacing
-					for (const FUniformBufferMemberInfo& MemberInfo : UniformBufferInfo.Members)
-					{
-						if (const TCHAR* SearchEnd = MatchStructMemberName(MemberInfo.NameAsStructMember, SearchPtr))
-						{
-							const int32 FoundTextLen = SearchEnd - SearchPtr;
-							const int32 ReplacementTextLen = MemberInfo.GlobalName.Len();
-
-							check(FoundTextLen >= ReplacementTextLen);
-
-							const TCHAR* SourceStart = GetData(MemberInfo.GlobalName);
-							TCHAR* const DestinationStart = SearchPtr;
-
-							int32 Index = 0;
-							for (; Index < ReplacementTextLen; Index++)
-							{
-								DestinationStart[Index] = SourceStart[Index];
-							}
-
-							// The shader preprocessor inserts spaces after defines
-							// #define ReflectionStruct OpaqueBasePass.Shared.Reflection
-							// 'ReflectionStruct.SkyLightCubemapBrightness' becomes 'OpaqueBasePass.Shared.Reflection .SkyLightCubemapBrightness' after MCPP
-							// In order to convert this struct member reference into a globally unique variable we move the spaces to the end
-							// 'OpaqueBasePass.Shared.Reflection .SkyLightCubemapBrightness' -> 'OpaqueBasePass_Shared_Reflection_SkyLightCubemapBrightness '
-
-							for (; Index < FoundTextLen; Index++)
-							{
-								// If we passed MatchStructMemberName, it should not be possible to overwrite the null terminator
-								check(DestinationStart[Index] != '\0');
-								DestinationStart[Index] = ' ';
-							}
-
-							break;
-						}
-					}
-				}
-
-				SearchPtr = FCString::Strstr(SearchPtr + UniformBufferName.Len(), *UniformBufferName);
-			}
-		}
+		memmove(SourceStart + DestOffset, SourceStart + SourceOffset, MoveCount * sizeof(TCHAR));
+		DestOffset += MoveCount;
 	}
+	PreprocessedShaderSource.GetCharArray().SetNum(DestOffset, false);
 }
 
 ///////////////////////////////////////////////////////////////////////////////////////////////////////////////////////////////////////
