Patch order of entry point parameters in cross-compiled HLSL of D3D11 backend using HlslParser.

#rnx

diff --git a/Engine/Shaders/Public/Platform/D3D/D3DCommon.ush b/Engine/Shaders/Public/Platform/D3D/D3DCommon.ush
--- a/Engine/Shaders/Public/Platform/D3D/D3DCommon.ush
+++ b/Engine/Shaders/Public/Platform/D3D/D3DCommon.ush
@@ -5,7 +5,7 @@ D3DCommon.ush: Common defines for the D3D shader compiler
 =============================================================================*/
 
 // Update this GUID to invalidate D3D shaders and force a shader recompilation for D3D RHIs.
-#pragma message("UESHADERMETADATA_VERSION 9C818FFC-0B81-4BCB-BF16-CFA7E4094DE3")
+#pragma message("UESHADERMETADATA_VERSION 46268304-2E81-419C-8152-0154868CFB1E")
 
 #include "/Engine/Shared/HLSLReservedSpaces.h"
 
diff --git a/Engine/Source/Developer/ShaderCompilerCommon/Private/HlslUtils.cpp b/Engine/Source/Developer/ShaderCompilerCommon/Private/HlslUtils.cpp
--- a/Engine/Source/Developer/ShaderCompilerCommon/Private/HlslUtils.cpp
+++ b/Engine/Source/Developer/ShaderCompilerCommon/Private/HlslUtils.cpp
@@ -113,6 +113,27 @@ static inline TCHAR* MakeIndexedSemantic(CrossCompiler::FLinearAllocator* Alloca
 	return Allocator->Strdup(Out);
 }
 
+static bool CheckSimpleVectorType(const TCHAR* SimpleType)
+{
+	if (!FCString::Strncmp(SimpleType, TEXT("float"), 5))
+	{
+		SimpleType += 5;
+	}
+	else if (!FCString::Strncmp(SimpleType, TEXT("int"), 3))
+	{
+		SimpleType += 3;
+	}
+	else if (!FCString::Strncmp(SimpleType, TEXT("half"), 4))
+	{
+		SimpleType += 4;
+	}
+	else
+	{
+		return false;
+	}
+	return FChar::IsDigit(SimpleType[0]) && SimpleType[1] == 0;
+}
+
 struct FRemoveAlgorithm
 {
 	FString EntryPoint;
@@ -324,28 +345,6 @@ struct FRemoveAlgorithm
 		return true;
 	}
 
-	bool CheckSimpleVectorType(const TCHAR* SimpleType)
-	{
-		if (!FCString::Strncmp(SimpleType, TEXT("float"), 5))
-		{
-			SimpleType += 5;
-		}
-		else if (!FCString::Strncmp(SimpleType, TEXT("int"), 3))
-		{
-			SimpleType += 3;
-		}
-		else if (!FCString::Strncmp(SimpleType, TEXT("half"), 4))
-		{
-			SimpleType += 4;
-		}
-		else
-		{
-			return false;
-		}
-
-		return FChar::IsDigit(SimpleType[0]) && SimpleType[1] == 0;
-	}
-
 	CrossCompiler::AST::FDeclaratorList* CreateLocalVariable(const TCHAR* Type, const TCHAR* Name, CrossCompiler::AST::FExpression* Initializer = nullptr)
 	{
 		using namespace CrossCompiler::AST;
@@ -1436,6 +1435,232 @@ bool RemoveUnusedInputs(FString& InOutSourceCode, const TArray<FString>& InInput
 	return RemoveUnusedInputs(InOutSourceCode, Inputs, {}, EntryPoint, OutErrors);
 }
 
+struct FFindEntryPointParameters
+{
+	FString EntryPoint;
+	const bool bFindOutputSemantics;
+
+	bool bSuccess;
+	TArray<FString> Errors;
+	CrossCompiler::FLinearAllocator* Allocator;
+
+	TArray<FString> FoundSemantics;
+
+	FFindEntryPointParameters(const FStringView& EntryPoint, EShaderParameterStorageClass SemanticsStorageClass) :
+		EntryPoint(EntryPoint),
+		bFindOutputSemantics(SemanticsStorageClass == EShaderParameterStorageClass::Output),
+		bSuccess(false),
+		Allocator(nullptr)
+	{
+	}
+
+	CrossCompiler::AST::FFunctionDefinition* FindEntryPointAndPopulateSymbolTable(CrossCompiler::TLinearArray<CrossCompiler::AST::FNode*>& ASTNodes, TArray<CrossCompiler::AST::FStructSpecifier*>& OutMiniSymbolTable)
+	{
+		using namespace CrossCompiler::AST;
+		FFunctionDefinition* EntryFunction = nullptr;
+		for (int32 Index = 0; Index < ASTNodes.Num(); ++Index)
+		{
+			FNode* Node = ASTNodes[Index];
+			if (FDeclaratorList* DeclaratorList = Node->AsDeclaratorList())
+			{
+				if (FStructSpecifier* StructSpecifier = DeclaratorList->Type->Specifier->Structure)
+				{
+					// Skip unnamed structures
+					if (StructSpecifier->Name)
+					{
+						OutMiniSymbolTable.Add(StructSpecifier);
+					}
+				}
+			}
+			else if (FFunctionDefinition* FunctionDefinition = Node->AsFunctionDefinition())
+			{
+				if (FCString::Strcmp(*EntryPoint, FunctionDefinition->Prototype->Identifier) == 0)
+				{
+					EntryFunction = FunctionDefinition;
+				}
+			}
+		}
+
+		return EntryFunction;
+	}
+
+	CrossCompiler::AST::FStructSpecifier* FindStructSpecifier(TArray<CrossCompiler::AST::FStructSpecifier*>& MiniSymbolTable, const TCHAR* StructName)
+	{
+		for (auto* StructSpecifier : MiniSymbolTable)
+		{
+			if (!FCString::Strcmp(StructSpecifier->Name, StructName))
+			{
+				return StructSpecifier;
+			}
+		}
+		return nullptr;
+	}
+
+	bool ProcessParameterSemantic(CrossCompiler::AST::FSemanticSpecifier& SemanticSpecifier, const TCHAR* ParameterIdentifier, bool bIsOutput)
+	{
+		if (const TCHAR* SemanticName = SemanticSpecifier.Semantic)
+		{
+			if (bFindOutputSemantics == bIsOutput)
+			{
+				FoundSemantics.Add(SemanticName);
+			}
+			return true;
+		}
+		else
+		{
+			check(ParameterIdentifier);
+			Errors.Add(FString::Printf(TEXT("FindEntryPointParameters: Function %s with parameter %s doesn't have a valid semantic name"), *EntryPoint, ParameterIdentifier));
+			return false;
+		}
+	}
+
+	bool ProcessStructParameterSemantics(CrossCompiler::AST::FStructSpecifier& StructSpecifier, TArray<CrossCompiler::AST::FStructSpecifier*>& SymbolTable, bool bIsOutput)
+	{
+		using namespace CrossCompiler::AST;
+		for (FNode* MemberNode : StructSpecifier.Members)
+		{
+			FDeclaratorList* DeclList = MemberNode->AsDeclaratorList();
+			if (!DeclList)
+			{
+				continue;
+			}
+			for (FNode* MemberNodeDecl : DeclList->Declarations)
+			{
+				if (FDeclaration* MemberDecl = MemberNodeDecl->AsDeclaration())
+				{
+					if (FSemanticSpecifier* Semantic = MemberDecl->Semantic)
+					{
+						if (!ProcessParameterSemantic(*Semantic, MemberDecl->Identifier, bIsOutput))
+						{
+							return false;
+						}
+					}
+					else if (FStructSpecifier* SubStructSpecifier = FindStructSpecifier(SymbolTable, DeclList->Type->Specifier->TypeName))
+					{
+						if (!ProcessStructParameterSemantics(*SubStructSpecifier, SymbolTable, bIsOutput))
+						{
+							return false;
+						}
+					}
+					else if (CheckSimpleVectorType(DeclList->Type->Specifier->TypeName))
+					{
+						Errors.Add(FString::Printf(TEXT("FindEntryPointParameters: Function %s with parameter %s doesn't have a return semantic"), *EntryPoint, MemberDecl->Identifier));
+						return false;
+					}
+					else
+					{
+						Errors.Add(FString::Printf(TEXT("FindEntryPointParameters: Invalid return type %s for parameter %s in function %s"), DeclList->Type->Specifier->TypeName, MemberDecl->Identifier, *EntryPoint));
+						return false;
+					}
+				}
+			}
+		}
+		return true;
+	}
+
+	bool ProcessFunctionParameters(CrossCompiler::AST::FFunctionDefinition* EntryFunction, TArray<CrossCompiler::AST::FStructSpecifier*>& SymbolTable)
+	{
+		using namespace CrossCompiler::AST;
+		for (FNode* ParamNode : EntryFunction->Prototype->Parameters)
+		{
+			FParameterDeclarator* ParameterDeclarator = ParamNode->AsParameterDeclarator();
+			check(ParameterDeclarator);
+			const bool bIsOutput = ParameterDeclarator->Type->Qualifier.bOut;
+			if (FSemanticSpecifier* Semantic = ParameterDeclarator->Semantic)
+			{
+				if (!ProcessParameterSemantic(*Semantic, ParameterDeclarator->Identifier, bIsOutput))
+				{
+					return false;
+				}
+			}
+			else if (FStructSpecifier* StructSpecifier = FindStructSpecifier(SymbolTable, ParameterDeclarator->Type->Specifier->TypeName))
+			{
+				if (!ProcessStructParameterSemantics(*StructSpecifier, SymbolTable, bIsOutput))
+				{
+					return false;
+				}
+			}
+			else if (CheckSimpleVectorType(ParameterDeclarator->Type->Specifier->TypeName))
+			{
+				Errors.Add(FString::Printf(TEXT("FindEntryPointParameters: Function %s with parameter %s doesn't have a return semantic"), *EntryPoint, ParameterDeclarator->Identifier));
+				return false;
+			}
+			else
+			{
+				Errors.Add(FString::Printf(TEXT("FindEntryPointParameters: Invalid return type %s for parameter %s in function %s"), ParameterDeclarator->Type->Specifier->TypeName, ParameterDeclarator->Identifier, *EntryPoint));
+				return false;
+			}
+		}
+		return true;
+	}
+
+	void FindEntryPointParameters(CrossCompiler::TLinearArray<CrossCompiler::AST::FNode*>& ASTNodes)
+	{
+		using namespace CrossCompiler::AST;
+
+		// Find Entry point from original AST nodes
+		TArray<FStructSpecifier*> SymbolTable;
+		FFunctionDefinition* EntryFunction = FindEntryPointAndPopulateSymbolTable(ASTNodes, SymbolTable);
+		if (!EntryFunction)
+		{
+			Errors.Add(FString::Printf(TEXT("FindEntryPointParameters: Unable to find entry point %s"), *EntryPoint));
+			bSuccess = false;
+			return;
+		}
+
+		if (!ProcessFunctionParameters(EntryFunction, SymbolTable))
+		{
+			bSuccess = false;
+			return;
+		}
+
+		bSuccess = true;
+	}
+};
+
+bool FindEntryPointParameters(
+	const FString& InSourceCode,
+	const FString& InEntryPoint,
+	EShaderParameterStorageClass ParameterStorageClass,
+	TArray<FString>& OutParameterSemantics,
+	TArray<FString>& OutErrors)
+{
+	check(ParameterStorageClass == EShaderParameterStorageClass::Input || ParameterStorageClass == EShaderParameterStorageClass::Output);
+
+	FString DummyFilename(TEXT("/Engine/Private/FindEntryPointParameters.usf"));
+	FFindEntryPointParameters Data(InEntryPoint, ParameterStorageClass);
+	auto ResultCallbackFunction = [&Data](CrossCompiler::FLinearAllocator* Allocator, CrossCompiler::TLinearArray<CrossCompiler::AST::FNode*>& ASTNodes) -> void
+		{
+			Data.Allocator = Allocator;
+			Data.FindEntryPointParameters(ASTNodes);
+		};
+	CrossCompiler::FCompilerMessages Messages;
+	if (!CrossCompiler::Parser::Parse(InSourceCode, DummyFilename, Messages, {}, ResultCallbackFunction))
+	{
+		Data.Errors.Add(FString(TEXT("FindEntryPointParameters: Failed to parse HLSL source!")));
+		OutErrors = MoveTemp(Data.Errors);
+		for (auto& Message : Messages.MessageList)
+		{
+			OutErrors.Add(Message.Message);
+		}
+		return false;
+	}
+
+	for (auto& Message : Messages.MessageList)
+	{
+		OutErrors.Add(Message.Message);
+	}
+
+	if (Data.bSuccess)
+	{
+		OutParameterSemantics = MoveTemp(Data.FoundSemantics);
+		return true;
+	}
+
+	OutErrors = MoveTemp(Data.Errors);
+	return false;
+}
+
 struct FConvertFP32ToFP16 {
 	FString Filename;
 	FString GeneratedCode;
diff --git a/Engine/Source/Developer/ShaderCompilerCommon/Public/ShaderCompilerCommon.h b/Engine/Source/Developer/ShaderCompilerCommon/Public/ShaderCompilerCommon.h
--- a/Engine/Source/Developer/ShaderCompilerCommon/Public/ShaderCompilerCommon.h
+++ b/Engine/Source/Developer/ShaderCompilerCommon/Public/ShaderCompilerCommon.h
@@ -374,6 +374,22 @@ extern SHADERCOMPILERCOMMON_API bool RemoveUnusedInputs(
 
 extern SHADERCOMPILERCOMMON_API bool RemoveUnusedInputs(FString& InOutSourceCode, const TArray<FString>& InUsedInputs, FString& InOutEntryPoint, TArray<FString>& OutErrors);
 
+// Shader input/output parameter storage classes. Naming adopted from SPIR-V nomenclature.
+enum class EShaderParameterStorageClass
+{
+	Input,
+	Output,
+};
+
+// Returns the semantic names of all individual entry point parameters (i.e. all structure fields are inlined)
+extern SHADERCOMPILERCOMMON_API bool FindEntryPointParameters(
+	const FString& InSourceCode,
+	const FString& InEntryPoint,
+	EShaderParameterStorageClass ParameterStorageClass,
+	TArray<FString>& OutParameterSemantics,
+	TArray<FString>& OutErrors
+);
+
 extern SHADERCOMPILERCOMMON_API bool ConvertFromFP32ToFP16(FString& InOutSourceCode, TArray<FString>& OutErrors);
 
 enum class EShaderConductorTarget
diff --git a/Engine/Source/Developer/Windows/ShaderFormatD3D/Private/D3DShaderCompiler.cpp b/Engine/Source/Developer/Windows/ShaderFormatD3D/Private/D3DShaderCompiler.cpp
--- a/Engine/Source/Developer/Windows/ShaderFormatD3D/Private/D3DShaderCompiler.cpp
+++ b/Engine/Source/Developer/Windows/ShaderFormatD3D/Private/D3DShaderCompiler.cpp
@@ -507,7 +507,12 @@ static void PatchSpirvForPrecompilation(FSpirv& Spirv)
 }
 
 // @param StageVariablesStorageClass Must be SpvStorageClassOutput for vertex shaders and SpvStorageClassInput for pixel shaders.
-static bool PatchHlslWithReorderedIOVariables(FString& HlslSourceString, const FSpirv& Spirv, SpvStorageClass StageVariablesStorageClass)
+static bool PatchHlslWithReorderedIOVariables(
+	FString& HlslSourceString,
+	const FString& OriginalShaderSource,
+	const FString& OriginalEntryPoint,
+	SpvStorageClass StageVariablesStorageClass,
+	TArray<FShaderCompilerError>& OutErrors)
 {
 	check(StageVariablesStorageClass == SpvStorageClassInput || StageVariablesStorageClass == SpvStorageClassOutput);
 
@@ -538,8 +543,16 @@ static bool PatchHlslWithReorderedIOVariables(FString& HlslSourceString, const F
 	StageVariableDeclarationSource.ParseIntoArrayLines(StageVariableDeclarationLines);
 
 	// Parse variable names from SPIR-V input
-	TArray<FString> Variables;
-	ParseSpirvGlobalVariables(Spirv, StageVariablesStorageClass, Variables);
+	TArray<FString> Variables, ParsingErrors;
+	const EShaderParameterStorageClass ParameterStorageClass = (StageVariablesStorageClass == SpvStorageClassOutput ? EShaderParameterStorageClass::Output : EShaderParameterStorageClass::Input);
+	if (!FindEntryPointParameters(OriginalShaderSource, OriginalEntryPoint, ParameterStorageClass, Variables, ParsingErrors))
+	{
+		for (FString& Error : ParsingErrors)
+		{
+			OutErrors.Add(FShaderCompilerError(MoveTemp(Error)));
+		}
+		return false;
+	}
 
 	if (Variables.Num() != StageVariableDeclarationLines.Num())
 	{
@@ -554,7 +567,8 @@ static bool PatchHlslWithReorderedIOVariables(FString& HlslSourceString, const F
 	{
 		for (FString& SourceLine : StageVariableDeclarationLines)
 		{
-			if (SourceLine.Find(Variable, ESearchCase::CaseSensitive) != INDEX_NONE)
+			// Search for semantic name (always case insensitive) in current stage variable source line
+			if (SourceLine.Find(Variable, ESearchCase::IgnoreCase) != INDEX_NONE)
 			{
 				// Append source line for current variable at the end of sorted declaration string.
 				// Then empty this source line to avoid unnecessary string comparisons for next variables.
@@ -574,7 +588,12 @@ static bool PatchHlslWithReorderedIOVariables(FString& HlslSourceString, const F
 }
 
 // @todo-lh: use ANSI string class whenever UE core gets one
-static void PatchHlslForPrecompilation(TArray<ANSICHAR>& HlslSource, const EShaderFrequency Frequency, const FSpirv& Spirv)
+static void PatchHlslForPrecompilation(
+	TArray<ANSICHAR>& HlslSource,
+	const EShaderFrequency Frequency,
+	const FString& OriginalShaderSource,
+	const FString& OriginalEntryPoint,
+	TArray<FShaderCompilerError>& OutErrors)
 {
 	FString HlslSourceString;
 
@@ -618,8 +637,8 @@ static void PatchHlslForPrecompilation(TArray<ANSICHAR>& HlslSource, const EShad
 
 	if (Frequency == SF_Vertex)
 	{
-		// Ensure order of output variables remains the same as declared in SPIR-V input
-		PatchHlslWithReorderedIOVariables(HlslSourceString, Spirv, SpvStorageClassOutput);
+		// Ensure order of output variables remains the same as declared in original shader source
+		PatchHlslWithReorderedIOVariables(HlslSourceString, OriginalShaderSource, OriginalEntryPoint, SpvStorageClassOutput, OutErrors);
 	}
 	else if (Frequency == SF_Pixel)
 	{
@@ -638,8 +657,8 @@ static void PatchHlslForPrecompilation(TArray<ANSICHAR>& HlslSource, const EShad
 			}
 		}
 
-		// Ensure order of input variables remains the same as declared in SPIR-V input
-		PatchHlslWithReorderedIOVariables(HlslSourceString, Spirv, SpvStorageClassInput);
+		// Ensure order of input variables remains the same as declared in original shader source
+		PatchHlslWithReorderedIOVariables(HlslSourceString, OriginalShaderSource, OriginalEntryPoint, SpvStorageClassInput, OutErrors);
 	}
 
 	// Return new HLSL source
@@ -788,7 +807,7 @@ static bool CompileAndProcessD3DShaderFXCExt(
 			}
 
 			// Patch HLSL for workarounds to prevent potential additional FXC failures
-			PatchHlslForPrecompilation(CrossCompiledSource, Frequency, Spirv);
+			PatchHlslForPrecompilation(CrossCompiledSource, Frequency, PreprocessedShaderSource, EntryPointName, Output.Errors);
 
 			if (bDumpDebugInfo && CrossCompiledSource.Num() > 1)
 			{
