Add support for operator overloading in HlslParser.

#rnx

diff --git a/Engine/Source/Developer/ShaderCompilerCommon/Private/HlslAST.cpp b/Engine/Source/Developer/ShaderCompilerCommon/Private/HlslAST.cpp
--- a/Engine/Source/Developer/ShaderCompilerCommon/Private/HlslAST.cpp
+++ b/Engine/Source/Developer/ShaderCompilerCommon/Private/HlslAST.cpp
@@ -569,6 +569,7 @@ namespace CrossCompiler
 			ReturnSemantic(nullptr),
 			Parameters(InAllocator),
 			bIsDefinition(false),
+			bIsOperator(false),
 			bIsStatic(false)
 		{
 		}
@@ -587,6 +588,10 @@ namespace CrossCompiler
 			{
 				Writer << ScopeIdentifier << TEXT("::");
 			}
+			if (bIsOperator)
+			{
+				Writer << TEXT("operator");
+			}
 			Writer << Identifier;
 			Writer << (TCHAR)'(';
 			bool bFirst = true;
diff --git a/Engine/Source/Developer/ShaderCompilerCommon/Private/HlslAST.h b/Engine/Source/Developer/ShaderCompilerCommon/Private/HlslAST.h
--- a/Engine/Source/Developer/ShaderCompilerCommon/Private/HlslAST.h
+++ b/Engine/Source/Developer/ShaderCompilerCommon/Private/HlslAST.h
@@ -652,6 +652,7 @@ namespace CrossCompiler
 			TLinearArray<FNode*> Parameters;
 
 			bool bIsDefinition;
+			bool bIsOperator;
 			bool bIsStatic;
 
 			//Signature
diff --git a/Engine/Source/Developer/ShaderCompilerCommon/Private/HlslLexer.cpp b/Engine/Source/Developer/ShaderCompilerCommon/Private/HlslLexer.cpp
--- a/Engine/Source/Developer/ShaderCompilerCommon/Private/HlslLexer.cpp
+++ b/Engine/Source/Developer/ShaderCompilerCommon/Private/HlslLexer.cpp
@@ -439,6 +439,7 @@ namespace CrossCompiler
 				InsertToken(TEXT("inline"), EHlslToken::Inline);
 				InsertToken(TEXT("typedef"), EHlslToken::Typedef);
 				InsertToken(TEXT("packoffset"), EHlslToken::PackOffset);
+				InsertToken(TEXT("operator"), EHlslToken::Operator);
 			}
 		} GStaticInitializer;
 	}
diff --git a/Engine/Source/Developer/ShaderCompilerCommon/Private/HlslLexer.h b/Engine/Source/Developer/ShaderCompilerCommon/Private/HlslLexer.h
--- a/Engine/Source/Developer/ShaderCompilerCommon/Private/HlslLexer.h
+++ b/Engine/Source/Developer/ShaderCompilerCommon/Private/HlslLexer.h
@@ -278,6 +278,7 @@ namespace CrossCompiler
 		PackOffset,
 
 		Identifier,
+		Operator, // HLSL2021
 		Literal,
 		BoolConstant,
 		StringConstant,	// C-style "string"
diff --git a/Engine/Source/Developer/ShaderCompilerCommon/Private/HlslParser.cpp b/Engine/Source/Developer/ShaderCompilerCommon/Private/HlslParser.cpp
--- a/Engine/Source/Developer/ShaderCompilerCommon/Private/HlslParser.cpp
+++ b/Engine/Source/Developer/ShaderCompilerCommon/Private/HlslParser.cpp
@@ -948,9 +948,9 @@ namespace CrossCompiler
 			}
 
 			// Greedily eat a member declaration, and backtrack if it's a member function
-			auto OriginalToken = Parser.Scanner.GetCurrentTokenIndex();
+			const uint32 OriginalToken = Parser.Scanner.GetCurrentTokenIndex();
 			AST::FDeclaratorList* Declaration = nullptr;
-			auto Result = ParseGeneralDeclarationNoSemicolon(Parser, SymbolScope, 0, EDF_CONST_ROW_MAJOR | EDF_SEMANTIC | EDF_MULTIPLE | EDF_TEXTURE_SAMPLER_OR_BUFFER | EDF_INTERPOLATION, Allocator, &Declaration);
+			EParseResult Result = ParseGeneralDeclarationNoSemicolon(Parser, SymbolScope, 0, EDF_CONST_ROW_MAJOR | EDF_SEMANTIC | EDF_MULTIPLE | EDF_TEXTURE_SAMPLER_OR_BUFFER | EDF_INTERPOLATION, Allocator, &Declaration);
 			bool bTryMemberFunction = false;
 			if (Result == EParseResult::Error)
 			{
@@ -1115,11 +1115,98 @@ namespace CrossCompiler
 		return EParseResult::Matched;
 	}
 
+	static bool IsOverloadableOperatorToken(EHlslToken Token)
+	{
+		switch (Token)
+		{
+		// Math
+		case EHlslToken::Plus:
+		case EHlslToken::PlusEqual:
+		case EHlslToken::Minus:
+		case EHlslToken::MinusEqual:
+		case EHlslToken::Times:
+		case EHlslToken::TimesEqual:
+		case EHlslToken::Div:
+		case EHlslToken::DivEqual:
+		case EHlslToken::Mod:
+		case EHlslToken::ModEqual:
+
+		// Logical
+		case EHlslToken::EqualEqual:
+		case EHlslToken::NotEqual:
+		case EHlslToken::Lower:
+		case EHlslToken::LowerEqual:
+		case EHlslToken::Greater:
+		case EHlslToken::GreaterEqual:
+
+		// Bit
+		case EHlslToken::LowerLower:
+		case EHlslToken::LowerLowerEqual:
+		case EHlslToken::GreaterGreater:
+		case EHlslToken::GreaterGreaterEqual:
+		case EHlslToken::And:
+		case EHlslToken::AndEqual:
+		case EHlslToken::Or:
+		case EHlslToken::OrEqual:
+		case EHlslToken::Xor:
+		case EHlslToken::XorEqual:
+		case EHlslToken::Not:
+		case EHlslToken::Neg:
+
+		// Statements
+		case EHlslToken::Equal:
+
+		// Unary
+		case EHlslToken::PlusPlus:
+		case EHlslToken::MinusMinus:
+			return true;
+
+		default:
+			return false;
+		}
+	}
+
+	static const TCHAR* TryParseOverloadableOperator(FHlslParser& Parser, FLinearAllocator* Allocator)
+	{
+		const FHlslToken* Token = Parser.Scanner.GetCurrentToken();
+		if (!Token)
+		{
+			return nullptr;
+		}
+	
+		if (Parser.Scanner.MatchToken(EHlslToken::LeftParenthesis))
+		{
+			// Match invocation operator '()'
+			if (!Parser.Scanner.MatchToken(EHlslToken::RightParenthesis))
+			{
+				Parser.Scanner.SourceError(TEXT("Missing closing ')' token for invocation operator '()'"));
+				return nullptr;
+			}
+			return TEXT("()");
+		}
+		else if (Parser.Scanner.MatchToken(EHlslToken::LeftSquareBracket))
+		{
+			// Match subscript operator '[]'
+			if (!Parser.Scanner.MatchToken(EHlslToken::RightSquareBracket))
+			{
+				Parser.Scanner.SourceError(TEXT("Missing closing ']' token for subscript operator '[]'"));
+				return nullptr;
+			}
+			return TEXT("[]");
+		}
+		else if (IsOverloadableOperatorToken(Token->Token))
+		{
+			Parser.Scanner.Advance();
+			return Allocator->Strdup(Token->String);
+		}
+		return nullptr;
+	}
+
 	EParseResult ParseFunctionDeclarator(FHlslParser& Parser, FLinearAllocator* Allocator, AST::FFunction** OutFunction)
 	{
-		auto OriginalToken = Parser.Scanner.GetCurrentTokenIndex();
+		uint32 OriginalToken = Parser.Scanner.GetCurrentTokenIndex();
 		AST::FTypeSpecifier* TypeSpecifier = nullptr;
-		auto Result = ParseGeneralType(Parser.Scanner, ETF_BUILTIN_NUMERIC | ETF_SAMPLER_TEXTURE_BUFFER | ETF_USER_TYPES | ETF_ERROR_IF_NOT_USER_TYPE | ETF_VOID, Parser.CurrentScope, Allocator, &TypeSpecifier);
+		EParseResult Result = ParseGeneralType(Parser.Scanner, ETF_BUILTIN_NUMERIC | ETF_SAMPLER_TEXTURE_BUFFER | ETF_USER_TYPES | ETF_ERROR_IF_NOT_USER_TYPE | ETF_VOID, Parser.CurrentScope, Allocator, &TypeSpecifier);
 		if (Result == EParseResult::NotMatched)
 		{
 			Parser.Scanner.SetCurrentTokenIndex(OriginalToken);
@@ -1132,25 +1219,45 @@ namespace CrossCompiler
 
 		check(Result == EParseResult::Matched);
 
-		const FHlslToken* Identifier = Parser.Scanner.GetCurrentToken();
-		if (!Parser.Scanner.MatchToken(EHlslToken::Identifier))
+		const TCHAR* OperatorIdentifier = nullptr;
+		const FHlslToken* Identifier = nullptr;
+
+		auto MatchIdentifierOrOperator = [OriginalToken, Allocator](FHlslParser& Parser, const FHlslToken*& Identifier, const TCHAR*& OperatorIdentifier) -> EParseResult
 		{
-			// This could be an error... But we should allow testing for a global variable before any rash decisions
-			Parser.Scanner.SetCurrentTokenIndex(OriginalToken);
-			return EParseResult::NotMatched;
+			if (Parser.Scanner.MatchToken(EHlslToken::Operator))
+			{
+				OperatorIdentifier = TryParseOverloadableOperator(Parser, Allocator);
+				if (!OperatorIdentifier)
+				{
+					Parser.Scanner.SourceError(TEXT("Expected operator token after 'operator' keyword!\n"));
+					return ParseResultError();
+				}
+			}
+			else if (!Parser.Scanner.MatchToken(EHlslToken::Identifier))
+			{
+				// This could be an error... But we should allow testing for a global variable before any rash decisions
+				Parser.Scanner.SetCurrentTokenIndex(OriginalToken);
+				return EParseResult::NotMatched;
+			}
+			return EParseResult::Matched;
+		};
+
+		Identifier = Parser.Scanner.GetCurrentToken();
+		Result = MatchIdentifierOrOperator(Parser, Identifier, OperatorIdentifier);
+		if (Result != EParseResult::Matched)
+		{
+			return Result;
 		}
 
 		// Check for scoped function declarations, e.g. "MyStruct::MyFunction() ..."
 		const FHlslToken* ScopeIdentifier = nullptr;
-		if (Parser.Scanner.MatchToken(EHlslToken::ColonColon))
+		if (OperatorIdentifier == nullptr && Parser.Scanner.MatchToken(EHlslToken::ColonColon))
 		{
 			ScopeIdentifier = Identifier;
-			Identifier = Parser.Scanner.GetCurrentToken();
-			if (!Parser.Scanner.MatchToken(EHlslToken::Identifier))
+			Result = MatchIdentifierOrOperator(Parser, Identifier, OperatorIdentifier);
+			if (Result != EParseResult::Matched)
 			{
-				// This could be an error... But we should allow testing for a global variable before any rash decisions
-				Parser.Scanner.SetCurrentTokenIndex(OriginalToken);
-				return EParseResult::NotMatched;
+				return Result;
 			}
 		}
 
@@ -1164,7 +1271,15 @@ namespace CrossCompiler
 		// At this point, any unknown identifiers could be a type being used without forward declaring, so it's a real error
 
 		auto* Function = new(Allocator) AST::FFunction(Allocator, Identifier->SourceInfo);
-		Function->Identifier = Allocator->Strdup(Identifier->String);
+		if (OperatorIdentifier)
+		{
+			Function->Identifier = OperatorIdentifier;
+			Function->bIsOperator = true;
+		}
+		else
+		{
+			Function->Identifier = Allocator->Strdup(Identifier->String);
+		}
 		if (ScopeIdentifier)
 		{
 			Function->ScopeIdentifier = Allocator->Strdup(ScopeIdentifier->String);
