Make DXC compile arguments consistent between Vulkan backend and CompileDXC.bat script.

#rnx

diff --git a/Engine/Source/Developer/ShaderCompilerCommon/Private/ShaderConductorContext.cpp b/Engine/Source/Developer/ShaderCompilerCommon/Private/ShaderConductorContext.cpp
--- a/Engine/Source/Developer/ShaderCompilerCommon/Private/ShaderConductorContext.cpp
+++ b/Engine/Source/Developer/ShaderCompilerCommon/Private/ShaderConductorContext.cpp
@@ -385,58 +385,33 @@ namespace CrossCompiler
 		}
 	}
 
-	static void ConvertScOptions(FShaderConductorContext::FShaderConductorIntermediates& Intermediates, const FShaderConductorOptions& InOptions, ShaderConductor::Compiler::Options& OutOptions, bool bIgnoreCustomDxcArgs = false)
+	static void AppendDxcArguments(const FShaderConductorOptions& InOptions, TArray<const ANSICHAR*>& DxcArguments)
 	{
-		// Validate input shader model with respect to certain language features.
-		checkf(
-			(!InOptions.bEnable16bitTypes || InOptions.ShaderModel >= FHlslShaderModel{ 6, 2 }),
-			TEXT("DXC option '-enable-16bit-types' only supported with SM6.2+ but SM%u.%u was specified"),
-			InOptions.ShaderModel.Major, InOptions.ShaderModel.Minor
-		);
-
-		OutOptions.removeUnusedGlobals = InOptions.bRemoveUnusedGlobals;
-		OutOptions.packMatricesInRowMajor = InOptions.bPackMatricesInRowMajor;
-		OutOptions.enable16bitTypes = InOptions.bEnable16bitTypes;
-		OutOptions.enableDebugInfo = InOptions.bEnableDebugInfo;
-		OutOptions.disableOptimizations = InOptions.bDisableOptimizations;
-		OutOptions.enableFMAPass = InOptions.bEnableFMAPass;
-		OutOptions.enableSeparateSamplers = InOptions.bEnableSeparateSamplersInGlsl;
-		OutOptions.remapAttributeLocations = InOptions.bRemapAttributeLocations;
-		OutOptions.shaderModel = ShaderConductor::Compiler::ShaderModel
-		{
-			static_cast<uint8>(InOptions.ShaderModel.Major),
-			static_cast<uint8>(InOptions.ShaderModel.Minor)
-		};
-
-		TArray<ANSICHAR const*>& DxcArgRefs = Intermediates.DxcArgRefs;
-
-		DxcArgRefs.Empty();
-
 		// Select language version
-		DxcArgRefs.Add("-spirv");
-		DxcArgRefs.Add("-Qunused-arguments");
+		DxcArguments.Add("-spirv");
+		DxcArguments.Add("-Qunused-arguments");
 
 		switch (InOptions.HlslVersion)
 		{
 		case 2015:
-			DxcArgRefs.Add("-HV");
-			DxcArgRefs.Add("2015");
+			DxcArguments.Add("-HV");
+			DxcArguments.Add("2015");
 			break;
 		case 2016:
-			DxcArgRefs.Add("-HV");
-			DxcArgRefs.Add("2016");
+			DxcArguments.Add("-HV");
+			DxcArguments.Add("2016");
 			break;
 		case 2017:
-			DxcArgRefs.Add("-HV");
-			DxcArgRefs.Add("2017");
+			DxcArguments.Add("-HV");
+			DxcArguments.Add("2017");
 			break;
 		case 2018:
-			DxcArgRefs.Add("-HV");
-			DxcArgRefs.Add("2018");
+			DxcArguments.Add("-HV");
+			DxcArguments.Add("2018");
 			break;
 		case 2021:
-			DxcArgRefs.Add("-HV");
-			DxcArgRefs.Add("2021");
+			DxcArguments.Add("-HV");
+			DxcArguments.Add("2021");
 			break;
 		default:
 			checkf(false, TEXT("Invalid HLSL version: expected 2015, 2016, 2017, 2018, or 2021 but %u was specified"), InOptions.HlslVersion);
@@ -446,33 +421,32 @@ namespace CrossCompiler
 		// Add additional DXC arguments that are not exposed by ShaderConductor API directly
 		if (!InOptions.bDisableScalarBlockLayout)
 		{
-			DxcArgRefs.Add("-fvk-use-scalar-layout");
+			DxcArguments.Add("-fvk-use-scalar-layout");
 		}
 		if (InOptions.bPreserveStorageInput)
 		{
-			DxcArgRefs.Add("-fspv-preserve-storage-input");
+			DxcArguments.Add("-fspv-preserve-storage-input");
+		}
+		if (InOptions.bForceStorageImageFormat)
+		{
+			DxcArguments.Add("-fvk-force-storage-image-format");
 		}
-        if (InOptions.bForceStorageImageFormat)
-        {
-            DxcArgRefs.Add("-fvk-force-storage-image-format");
-        }
 		if (InOptions.bSvPositionImplicitInvariant)
 		{
-			DxcArgRefs.Add("-fspv-svposition-implicit-invariant");
+			DxcArguments.Add("-fspv-svposition-implicit-invariant");
 		}
 		if (InOptions.bSupportPreciseOutputs)
 		{
-			DxcArgRefs.Add("-fspv-support-precise-outputs");
+			DxcArguments.Add("-fspv-support-precise-outputs");
 		}
 
-
 		using ETargetEnvironment = CrossCompiler::FShaderConductorOptions::ETargetEnvironment;
-		if (OutOptions.enable16bitTypes)
+		if (InOptions.bEnable16bitTypes)
 		{
-			DxcArgRefs.Add("-fspv-target-env=universal1.5");
+			DxcArguments.Add("-fspv-target-env=universal1.5");
 
 			// ShaderConductor.cpp forgot to pipedown enable16bitTypes, so work arround by adding the parameter manually in here.
-			DxcArgRefs.Add("-enable-16bit-types");
+			DxcArguments.Add("-enable-16bit-types");
 		}
 		else
 		{
@@ -481,19 +455,49 @@ namespace CrossCompiler
 			default:
 				checkf(false, TEXT("Unexpected SPIR-V target environment: %d"), (uint32)InOptions.TargetEnvironment);
 			case ETargetEnvironment::Vulkan_1_0:
-				DxcArgRefs.Add("-fspv-target-env=vulkan1.0");
+				DxcArguments.Add("-fspv-target-env=vulkan1.0");
 				break;
 			case ETargetEnvironment::Vulkan_1_1:
-				DxcArgRefs.Add("-fspv-target-env=vulkan1.1");
+				DxcArguments.Add("-fspv-target-env=vulkan1.1");
 				break;
 			case ETargetEnvironment::Vulkan_1_2:
-				DxcArgRefs.Add("-fspv-target-env=vulkan1.2");
+				DxcArguments.Add("-fspv-target-env=vulkan1.2");
 				break;
 			case ETargetEnvironment::Vulkan_1_3:
-				DxcArgRefs.Add("-fspv-target-env=vulkan1.3");
+				DxcArguments.Add("-fspv-target-env=vulkan1.3");
 				break;
 			}
 		}
+	}
+
+	static void ConvertScOptions(FShaderConductorContext::FShaderConductorIntermediates& Intermediates, const FShaderConductorOptions& InOptions, ShaderConductor::Compiler::Options& OutOptions, bool bIgnoreCustomDxcArgs = false)
+	{
+		// Validate input shader model with respect to certain language features.
+		checkf(
+			(!InOptions.bEnable16bitTypes || InOptions.ShaderModel >= FHlslShaderModel{ 6, 2 }),
+			TEXT("DXC option '-enable-16bit-types' only supported with SM6.2+ but SM%u.%u was specified"),
+			InOptions.ShaderModel.Major, InOptions.ShaderModel.Minor
+		);
+
+		OutOptions.removeUnusedGlobals = InOptions.bRemoveUnusedGlobals;
+		OutOptions.packMatricesInRowMajor = InOptions.bPackMatricesInRowMajor;
+		OutOptions.enable16bitTypes = InOptions.bEnable16bitTypes;
+		OutOptions.enableDebugInfo = InOptions.bEnableDebugInfo;
+		OutOptions.disableOptimizations = InOptions.bDisableOptimizations;
+		OutOptions.enableFMAPass = InOptions.bEnableFMAPass;
+		OutOptions.enableSeparateSamplers = InOptions.bEnableSeparateSamplersInGlsl;
+		OutOptions.remapAttributeLocations = InOptions.bRemapAttributeLocations;
+		OutOptions.shaderModel = ShaderConductor::Compiler::ShaderModel
+		{
+			static_cast<uint8>(InOptions.ShaderModel.Major),
+			static_cast<uint8>(InOptions.ShaderModel.Minor)
+		};
+
+		TArray<ANSICHAR const*>& DxcArgRefs = Intermediates.DxcArgRefs;
+
+		DxcArgRefs.Empty();
+
+		AppendDxcArguments(InOptions, DxcArgRefs);
 
 		if (!InOptions.SpirvCustomOptimizationPasses.IsEmpty())
 		{
@@ -847,6 +851,39 @@ namespace CrossCompiler
 		return (Intermediates->ShaderSource.Num() > 0 ? (Intermediates->ShaderSource.Num() - 1) : 0);
 	}
 
+	static const TCHAR* GetHlslShaderModelProfile(ShaderConductor::ShaderStage Stage)
+	{
+		switch (Stage)
+		{
+		case ShaderConductor::ShaderStage::VertexShader:	return TEXT("vs");
+		case ShaderConductor::ShaderStage::PixelShader:		return TEXT("ps");
+		case ShaderConductor::ShaderStage::GeometryShader:	return TEXT("gs");
+		case ShaderConductor::ShaderStage::HullShader:		return TEXT("hl");
+		case ShaderConductor::ShaderStage::DomainShader:	return TEXT("ds");
+		case ShaderConductor::ShaderStage::ComputeShader:	return TEXT("cs");
+		default:											return TEXT("lib");
+		}
+	}
+
+	FString FShaderConductorContext::GenerateDxcArguments(const FShaderConductorOptions& Options) const
+	{
+		FString CmdLineArgs = FString::Printf(
+			TEXT("-E %s -T %s_%d_%d"),
+			ANSI_TO_TCHAR(Intermediates->EntryPoint.GetData()), GetHlslShaderModelProfile(Intermediates->Stage), (int32)Options.ShaderModel.Major, (int32)Options.ShaderModel.Minor
+		);
+
+		TArray<const ANSICHAR*> DxcArguments;
+		AppendDxcArguments(Options, DxcArguments);
+
+		for (const ANSICHAR* Argument : DxcArguments)
+		{
+			CmdLineArgs += TEXT(" ");
+			CmdLineArgs += ANSI_TO_TCHAR(Argument);
+		}
+
+		return CmdLineArgs;
+	}
+
 	void FShaderConductorContext::ConvertCompileErrors(TArray<FString>&& ErrorStringLines, TArray<FShaderCompilerError>& OutErrors)
 	{
 		// Returns whether the specified line in the 'ErrorStringLines' array has a line marker.
diff --git a/Engine/Source/Developer/ShaderCompilerCommon/Public/ShaderConductorContext.h b/Engine/Source/Developer/ShaderCompilerCommon/Public/ShaderConductorContext.h
--- a/Engine/Source/Developer/ShaderCompilerCommon/Public/ShaderConductorContext.h
+++ b/Engine/Source/Developer/ShaderCompilerCommon/Public/ShaderConductorContext.h
@@ -232,6 +232,9 @@ namespace CrossCompiler
 		/** Returns a length of the internal loaded sources (excluding the null terminator). This is automatically updated when RewriteHlsl() is called. */
 		int32 GetSourceLength() const;
 
+		/** Returns the DXC command line arguments for the specified options. This does not include an output file, i.e. "-Fo" argument is not included. */
+		FString GenerateDxcArguments(const FShaderConductorOptions& Options) const;
+
 		/** Returns the list of current compile errors. */
 		inline const TArray<FShaderCompilerError>& GetErrors() const
 		{
diff --git a/Engine/Source/Developer/VulkanShaderFormat/Private/VulkanShaderCompiler.cpp b/Engine/Source/Developer/VulkanShaderFormat/Private/VulkanShaderCompiler.cpp
--- a/Engine/Source/Developer/VulkanShaderFormat/Private/VulkanShaderCompiler.cpp
+++ b/Engine/Source/Developer/VulkanShaderFormat/Private/VulkanShaderCompiler.cpp
@@ -1764,44 +1764,14 @@ static void Patch64bitSamplers(FVulkanSpirv& Spirv)
 	}
 }
 
-static FString VulkanGetShaderProfileDXC(const FVulkanShaderCompilerInternalState& InternalState, const CrossCompiler::FShaderConductorOptions Options)
-{
-	const TCHAR* ShaderProfile = TEXT("unknown");
-	switch (InternalState.GetShaderFrequency())
-	{
-	case SF_Vertex:
-		ShaderProfile = TEXT("vs");
-		break;
-
-	case SF_Pixel:
-		ShaderProfile = TEXT("ps");
-		break;
-
-	case SF_Geometry:
-		ShaderProfile = TEXT("gs");
-		break;
-
-	case SF_Compute:
-		ShaderProfile = TEXT("cs");
-		break;
-
-	case SF_RayGen:
-	case SF_RayMiss:
-	case SF_RayHitGroup:
-	case SF_RayCallable:
-		return TEXT("lib_6_3");
-	}
-
-	return FString::Printf(TEXT("%s_%d_%d"), ShaderProfile, Options.ShaderModel.Major, Options.ShaderModel.Minor); 
-}
-
 static void VulkanCreateDXCCompileBatchFiles(
+	const CrossCompiler::FShaderConductorContext& CompilerContext,
 	const FVulkanShaderCompilerInternalState& InternalState,
-	const CrossCompiler::FShaderConductorOptions Options)
+	const CrossCompiler::FShaderConductorOptions& Options)
 {
-	FString USFFilename = InternalState.Input.GetSourceFilename();
-	FString SPVFilename = FPaths::GetBaseFilename(USFFilename) + TEXT(".DXC.spv");
-	FString GLSLFilename = FPaths::GetBaseFilename(USFFilename) + TEXT(".SPV.glsl");
+	const FString USFFilename = InternalState.Input.GetSourceFilename();
+	const FString SPVFilename = FPaths::GetBaseFilename(USFFilename) + TEXT(".DXC.spv");
+	const FString GLSLFilename = FPaths::GetBaseFilename(USFFilename) + TEXT(".SPV.glsl");
 
 	FString DxcPath = FPaths::ConvertRelativePathToFull(FPaths::EngineDir());
 
@@ -1811,32 +1781,10 @@ static void VulkanCreateDXCCompileBatchFiles(
 	FString DxcFilename = FPaths::Combine(DxcPath, TEXT("dxc.exe"));
 	FPaths::MakePlatformFilename(DxcFilename);
 
-	const TCHAR* VulkanVersion = TEXT("vulkanUNKNOWN");
-	if (Options.TargetEnvironment == CrossCompiler::FShaderConductorOptions::ETargetEnvironment::Vulkan_1_0)
-	{
-		VulkanVersion = TEXT("vulkan1.0");
-	}
-	else if (Options.TargetEnvironment == CrossCompiler::FShaderConductorOptions::ETargetEnvironment::Vulkan_1_1)
-	{
-		VulkanVersion = TEXT("vulkan1.1");
-	}
-	else if (Options.TargetEnvironment == CrossCompiler::FShaderConductorOptions::ETargetEnvironment::Vulkan_1_2)
-	{
-		VulkanVersion = TEXT("vulkan1.2");
-	}
-	else if (Options.TargetEnvironment == CrossCompiler::FShaderConductorOptions::ETargetEnvironment::Vulkan_1_3)
-	{
-		VulkanVersion = TEXT("vulkan1.3");
-	}
-	else
-	{
-		ensure(false);
-	}
-
-	FString ShaderProfile = VulkanGetShaderProfileDXC(InternalState, Options);
-
 	// CompileDXC.bat
 	{
+		const FString DxcArguments = CompilerContext.GenerateDxcArguments(Options);
+
 		FString BatchFileContents =  FString::Printf(
 			TEXT(
 				"@ECHO OFF\n"
@@ -1847,7 +1795,7 @@ static void VulkanCreateDXCCompileBatchFiles(
 				"\tGOTO :END\n"
 				")\n"
 				"ECHO Compiling with DXC...\n"
-				"%%DXC%% -HV %d -T %s -E %s -spirv -fspv-target-env=%s -Fo %s %s\n"
+				"%%DXC%% %s -Fo %s %s\n"
 				"WHERE %%SPIRVCROSS%%\n"
 				"IF %%ERRORLEVEL%% NEQ 0 (\n"
 				"\tECHO spirv-cross.exe not found in Path environment variable, please build it from source [REDACTED_PROFILE_URL]/SPIRV-Cross\n"
@@ -1860,10 +1808,7 @@ static void VulkanCreateDXCCompileBatchFiles(
 			),
 			*DxcFilename,
 			*DxcPath,
-			Options.HlslVersion,
-			*ShaderProfile,
-			*InternalState.GetEntryPointName(),
-			VulkanVersion,
+			*DxcArguments,
 			*SPVFilename,
 			*USFFilename,
 			*GLSLFilename,
@@ -2056,7 +2001,7 @@ static bool CompileWithShaderConductor(
 
 	if (InternalState.bDebugDump)
 	{
-		VulkanCreateDXCCompileBatchFiles(InternalState, Options);
+		VulkanCreateDXCCompileBatchFiles(CompilerContext, InternalState, Options);
 	}
 
 	// Before the shader rewritter removes all traces of it, pull any WAVESIZE directives from the shader source
