[shaders] remove unused remnants of TEXT macro parsing code

diff --git a/Engine/Source/Developer/ShaderCompilerCommon/Private/ShaderCompilerCommon.cpp b/Engine/Source/Developer/ShaderCompilerCommon/Private/ShaderCompilerCommon.cpp
--- a/Engine/Source/Developer/ShaderCompilerCommon/Private/ShaderCompilerCommon.cpp
+++ b/Engine/Source/Developer/ShaderCompilerCommon/Private/ShaderCompilerCommon.cpp
@@ -768,11 +768,6 @@ void AddUnboundShaderParameterError(
 	AddNoteToDisplayShaderParameterStructureOnCppSide(CompilerInput.RootParametersStructure, CompilerOutput);
 }
 
-void RemoveUniformBuffersFromSource(const FShaderCompilerEnvironment& Environment, FString& PreprocessedShaderSource)
-{
-	CleanupUniformBufferCode(Environment, PreprocessedShaderSource);
-}
-
 struct FUniformBufferMemberInfo
 {
 	// eg View.WorldToClip
@@ -1396,110 +1391,6 @@ void CleanupUniformBufferCode(const FShaderCompilerEnvironment& Environment, FSt
 	PreprocessedShaderSource.GetCharArray().SetNum(DestOffset, false);
 }
 
-///////////////////////////////////////////////////////////////////////////////////////////////////////////////////////////////////////
-// Process TEXT() macro to convert them into GPU ASCII characters
-
-static FString ParseText(const TCHAR* StartPtr, const TCHAR*& EndPtr)
-{
-	const TCHAR* OpeningBracePtr = FindNextChar(StartPtr, '(');
-	check(OpeningBracePtr);
-
-	const TCHAR* ClosingBracePtr = FindMatchingClosingParenthesis(OpeningBracePtr + 1);
-	check(ClosingBracePtr);
-
-	FString Out;
-	if (OpeningBracePtr && ClosingBracePtr)
-	{
-		const TCHAR* CurrPtr = OpeningBracePtr;
-		do
-		{
-			Out += *CurrPtr;
-			CurrPtr++;
-		} while (CurrPtr != ClosingBracePtr+1);
-	}
-	EndPtr = ClosingBracePtr;
-	return Out;
-}
-
-static void ConvertTextToAsciiCharacter(const FString& InText, FString& OutText, FString& OutEncodedText)
-{
-	const uint32 CharCount = InText.Len();
-	OutEncodedText.Reserve(CharCount * 3); // ~2 digits per character + a comma
-	OutText = InText;
-	for (uint32 CharIt = 0; CharIt < CharCount; ++CharIt)
-	{
-		const char C = InText[CharIt];
-		OutEncodedText.AppendInt(uint8(C));
-		if (CharIt + 1 != CharCount)
-		{
-			OutEncodedText += ',';
-		}
-	}
-}
-
-struct FAssertParsingStack
-{
-	bool IsAssert() const { return BeginPtr != nullptr && EndPtr != nullptr; }
-	const TCHAR* BeginPtr = nullptr;
-	const TCHAR* EndPtr = nullptr;
-};
-
-static bool SearchText(const TCHAR*& InOut, FAssertParsingStack& OutAssertStack)
-{
-	const TCHAR* TextIdentifier = TEXT("TEXT(");
-	const TCHAR* AssertIdentifier = TEXT("UEReportAssertWithPayload(");
-
-	// Find the next TEXT() or UEReportAssertWithPayload()
-	const TCHAR* PrintfPtr = FCString::Strstr(InOut, TextIdentifier);
-	const TCHAR* AssertPtr = FCString::Strstr(InOut, AssertIdentifier);
-
-	// 1. Default is current TEXT parsing
-	InOut = PrintfPtr;
-
-	// 2. If current pointer is beyond the current assert context, reset the context
-	if (InOut >= OutAssertStack.EndPtr)
-	{
-		OutAssertStack.BeginPtr = nullptr;
-		OutAssertStack.EndPtr = nullptr;
-	}
-
-	// 3. If we are within an assert context, continue to part TEXT within that context
-	if (OutAssertStack.IsAssert())
-	{
-		// Nothing to do. InOut will be initialized by PrintfPtr, which is the next TEXT block.
-	}
-	// 4. If a new assert is detected, start a new assert context.
-	else if (AssertPtr && AssertPtr < PrintfPtr)
-	{
-		// Sanity check
-		check(!OutAssertStack.IsAssert());
-
-		// Check if the current assert is valid, i.e., containt a TEXT() argument
-		const TCHAR* EndPtr = nullptr;
-		const FString Tmp = ParseText(AssertPtr, EndPtr);
-		const bool bIsValid = FCString::Strstr(&Tmp[0], TextIdentifier) != nullptr; //-V547
-		if (bIsValid)
-		{
-			OutAssertStack.BeginPtr = AssertPtr;
-			OutAssertStack.EndPtr   = EndPtr;
-			InOut = FCString::Strstr(AssertPtr, TextIdentifier);
-		}
-		else
-		{
-			const uint32 LenAssertIdentifier = FString(AssertIdentifier).Len();
-			AssertPtr += LenAssertIdentifier;
-			return SearchText(AssertPtr, OutAssertStack);
-		}
-	}
-	return InOut != nullptr;
-}
-
-// Deprecated function -- logic now handled in PreprocessShader
-void TransformStringIntoCharacterArray(FString& PreprocessedShaderSource, TArray<FShaderDiagnosticData>* OutDiagnosticDatas)
-{
-	// Empty
-}
-
 ///////////////////////////////////////////////////////////////////////////////////////////////////////////////////////////////////////
 
 FString CreateShaderCompilerWorkerDirectCommandLine(const FShaderCompilerInput& Input)
diff --git a/Engine/Source/Developer/ShaderCompilerCommon/Public/ShaderCompilerCommon.h b/Engine/Source/Developer/ShaderCompilerCommon/Public/ShaderCompilerCommon.h
--- a/Engine/Source/Developer/ShaderCompilerCommon/Public/ShaderCompilerCommon.h
+++ b/Engine/Source/Developer/ShaderCompilerCommon/Public/ShaderCompilerCommon.h
@@ -331,17 +331,15 @@ extern SHADERCOMPILERCOMMON_API void AddUnboundShaderParameterError(
 // Convert generated UniformBuffer code and references into something the shader compilers can use.
 extern SHADERCOMPILERCOMMON_API void CleanupUniformBufferCode(const FShaderCompilerEnvironment& Environment, FString& PreprocessedShaderSource);
 
-// The cross compiler doesn't yet support struct initializers needed to construct static structs for uniform buffers
-// Replace all uniform buffer struct member references (View.WorldToClip) with a flattened name that removes the struct dependency (View_WorldToClip)
 UE_DEPRECATED(5.4, "RemoveUniformBuffersFromSource was renamed CleanupUniformBufferCode")
-extern SHADERCOMPILERCOMMON_API void RemoveUniformBuffersFromSource(const FShaderCompilerEnvironment& Environment, FString& PreprocessedShaderSource);
+inline void RemoveUniformBuffersFromSource(const FShaderCompilerEnvironment& Environment, FString& PreprocessedShaderSource) {}
+
 extern SHADERCOMPILERCOMMON_API const TCHAR* FindMatchingClosingBrace(const TCHAR* OpeningCharPtr);
 extern SHADERCOMPILERCOMMON_API const TCHAR* ParseHLSLSymbolName(const TCHAR* SearchString, FString& SymboName);
 extern SHADERCOMPILERCOMMON_API void ParseHLSLTypeName(const TCHAR* SearchString, const TCHAR*& TypeNameStartPtr, const TCHAR*& TypeNameEndPtr);
 
-// Processes TEXT macros
 UE_DEPRECATED(5.4, "TEXT macro processing is now handled in PreprocessShader --  TransformStringIntoCharacterArray is a no-op")
-extern SHADERCOMPILERCOMMON_API void TransformStringIntoCharacterArray(FString& PreprocessedShaderSource, TArray<FShaderDiagnosticData>* OutDiagnosticDatas=nullptr);
+inline void TransformStringIntoCharacterArray(FString& PreprocessedShaderSource, TArray<FShaderDiagnosticData>* OutDiagnosticDatas = nullptr) {}
 
 // Structure to hold forward declarations for a specific scope/namespace chain for the HlslParser
 struct FScopedDeclarations
