[shaders]
- add FShaderSource class which wraps source as populated by preprocessing and subsequently accessed by compilation and other debug features; this class automatically inserts zeroed padding such that 16-byte-wide SIMD string comparison operations do not require a non-SIMD tail to process any overhang.
- add typedefs for the string/view/character types and update preprocessing code to use these typedefs instead of the explicit types
- add explicit if constexprs in minifier code around char width to disable simd optimizations for char width != 2 (and subsequently skip the non-simd tail if char width == 2 since FShaderSource automatically adds the required padding)

diff --git a/Engine/Source/Developer/Apple/MetalShaderFormat/Private/MetalShaderCompiler.cpp b/Engine/Source/Developer/Apple/MetalShaderFormat/Private/MetalShaderCompiler.cpp
--- a/Engine/Source/Developer/Apple/MetalShaderFormat/Private/MetalShaderCompiler.cpp
+++ b/Engine/Source/Developer/Apple/MetalShaderFormat/Private/MetalShaderCompiler.cpp
@@ -766,10 +766,10 @@ bool PreprocessMetalShader(const FShaderCompilerInput& Input, const FShaderCompi
 	return UE::ShaderCompilerCommon::ExecuteShaderPreprocessingSteps(PreprocessOutput, Input, Environment);
 }
 
-void CompileMetalShader(const FShaderCompilerInput& Input, const FString& InPreprocessedSource, FShaderCompilerOutput& Output)
+void CompileMetalShader(const FShaderCompilerInput& Input, const FShaderPreprocessOutput& InPreprocessOutput, FShaderCompilerOutput& Output)
 {
 	FString EntryPointName = Input.EntryPointName;
-	FString PreprocessedSource = InPreprocessedSource;
+	FString PreprocessedSource(InPreprocessOutput.GetSourceViewWide());
 
 	FShaderParameterParser::FPlatformConfiguration PlatformConfiguration;
 	FShaderParameterParser ShaderParameterParser(PlatformConfiguration);
diff --git a/Engine/Source/Developer/Apple/MetalShaderFormat/Private/MetalShaderFormat.cpp b/Engine/Source/Developer/Apple/MetalShaderFormat/Private/MetalShaderFormat.cpp
--- a/Engine/Source/Developer/Apple/MetalShaderFormat/Private/MetalShaderFormat.cpp
+++ b/Engine/Source/Developer/Apple/MetalShaderFormat/Private/MetalShaderFormat.cpp
@@ -29,7 +29,7 @@ DEFINE_LOG_CATEGORY(LogMetalShaderCompiler)
 #define CHECK_METAL_COMPILER_TOOLCHAIN_SETUP 0
 
 extern bool PreprocessMetalShader(const FShaderCompilerInput& Input, const FShaderCompilerEnvironment& Environment, FShaderPreprocessOutput& PreprocessOutput);
-extern void CompileMetalShader(const FShaderCompilerInput& Input, const FString& InPreprocessedSource, FShaderCompilerOutput& Output);
+extern void CompileMetalShader(const FShaderCompilerInput& Input, const FShaderPreprocessOutput& InPreprocessOutput, FShaderCompilerOutput& Output);
 extern void OutputMetalDebugData(const FShaderCompilerInput& Input, const FShaderPreprocessOutput& PreprocessOutput, const FShaderCompilerOutput& Output);
 
 extern bool StripShader_Metal(TArray<uint8>& Code, class FString const& DebugPath, bool const bNative);
@@ -135,7 +135,7 @@ public:
 		const FString& WorkingDirectory) const override final
 	{
 		CheckShaderFormat(Input.ShaderFormat);
-		CompileMetalShader(Input, PreprocessOutput.GetSource(), Output);
+		CompileMetalShader(Input, PreprocessOutput, Output);
 	}
 
 	virtual void OutputDebugData(
diff --git a/Engine/Source/Developer/ShaderCompilerCommon/Private/ShaderCompilerCommon.cpp b/Engine/Source/Developer/ShaderCompilerCommon/Private/ShaderCompilerCommon.cpp
--- a/Engine/Source/Developer/ShaderCompilerCommon/Private/ShaderCompilerCommon.cpp
+++ b/Engine/Source/Developer/ShaderCompilerCommon/Private/ShaderCompilerCommon.cpp
@@ -219,14 +219,16 @@ void CullGlobalUniformBuffers(const TMap<FString, FUniformBufferEntry>& UniformB
 	}
 }
 
-static bool IsSpaceOrTabOrEOL(TCHAR Char)
+template <typename CharType>
+static bool IsSpaceOrTabOrEOL(CharType Char)
 {
 	return Char == ' ' || Char == '\t' || Char == '\n' || Char == '\r';
 }
 
-static const TCHAR* FindNextChar(const TCHAR* ReadStart, TCHAR SearchChar)
+template <typename StrCharType, typename SearchCharType>
+static const StrCharType* FindNextChar(const StrCharType* ReadStart, SearchCharType SearchChar)
 {
-	const TCHAR* SearchPtr = ReadStart;
+	const StrCharType* SearchPtr = ReadStart;
 	while (*SearchPtr && *SearchPtr != SearchChar)
 	{
 		SearchPtr++;
@@ -234,7 +236,8 @@ static const TCHAR* FindNextChar(const TCHAR* ReadStart, TCHAR SearchChar)
 	return SearchPtr;
 }
 
-const TCHAR* FindNextWhitespace(const TCHAR* StringPtr)
+template <typename CharType>
+const CharType* FindNextWhitespace(const CharType* StringPtr)
 {
 	while (*StringPtr && !IsSpaceOrTabOrEOL(*StringPtr))
 	{
@@ -251,7 +254,8 @@ const TCHAR* FindNextWhitespace(const TCHAR* StringPtr)
 	}
 }
 
-const TCHAR* FindNextNonWhitespace(const TCHAR* StringPtr)
+template <typename CharType>
+const CharType* FindNextNonWhitespace(const CharType* StringPtr)
 {
 	while (*StringPtr && IsSpaceOrTabOrEOL(*StringPtr))
 	{
@@ -266,7 +270,8 @@ const TCHAR* FindNextNonWhitespace(const TCHAR* StringPtr)
 	return nullptr;
 }
 
-const TCHAR* FindPreviousNonWhitespace(const TCHAR* StringPtr)
+template <typename CharType>
+const CharType* FindPreviousNonWhitespace(const CharType* StringPtr)
 {
 	do
 	{
@@ -281,36 +286,12 @@ const TCHAR* FindPreviousNonWhitespace(const TCHAR* StringPtr)
 	return nullptr;
 }
 
-const TCHAR* FindMatchingBlock(const TCHAR* OpeningCharPtr, char OpenChar, char CloseChar)
-{
-	const TCHAR* SearchPtr = OpeningCharPtr;
-	int32 Depth = 0;
-
-	while (*SearchPtr)
-	{
-		if (*SearchPtr == OpenChar)
-		{
-			Depth++;
-		}
-		else if (*SearchPtr == CloseChar)
-		{
-			if (Depth == 0)
-			{
-				return SearchPtr;
-			}
-
-			Depth--;
-		}
-		SearchPtr++;
-	}
-
-	return nullptr;
-}
-const TCHAR* FindMatchingClosingBrace(const TCHAR* OpeningCharPtr)			{ return FindMatchingBlock(OpeningCharPtr, '{', '}'); };
-const TCHAR* FindMatchingClosingParenthesis(const TCHAR* OpeningCharPtr)	{ return FindMatchingBlock(OpeningCharPtr, '(', ')'); };
+template <typename CharType>
+const CharType* FindMatchingClosingParenthesis(const CharType* OpeningCharPtr)	{ return FindMatchingBlock<CharType>(OpeningCharPtr, '(', ')'); };
 
 // See MSDN HLSL 'Symbol Name Restrictions' doc
-inline bool IsValidHLSLIdentifierCharacter(TCHAR Char)
+template <typename CharType>
+inline bool IsValidHLSLIdentifierCharacter(CharType Char)
 {
 	return (Char >= 'a' && Char <= 'z') ||
 		(Char >= 'A' && Char <= 'Z') ||
@@ -354,23 +335,24 @@ void ParseHLSLTypeName(const TCHAR* SearchString, const TCHAR*& TypeNameStartPtr
 	check(TypeNameEndPtr);
 }
 
-FStringView ParseHLSLSymbolName(const TCHAR* SearchString)
+template<typename CharType, typename ViewType>
+ViewType ParseHLSLSymbolName(const CharType* SearchString)
 {
-	const TCHAR* SymbolNameStartPtr = FindNextNonWhitespace(SearchString);
+	const CharType* SymbolNameStartPtr = FindNextNonWhitespace(SearchString);
 	check(SymbolNameStartPtr);
 
-	const TCHAR* SymbolNameEndPtr = SymbolNameStartPtr;
+	const CharType* SymbolNameEndPtr = SymbolNameStartPtr;
 	while (*SymbolNameEndPtr && IsValidHLSLIdentifierCharacter(*SymbolNameEndPtr))
 	{
 		SymbolNameEndPtr++;
 	}
 
-	return FStringView(SymbolNameStartPtr, SymbolNameEndPtr - SymbolNameStartPtr);
+	return ViewType(SymbolNameStartPtr, SymbolNameEndPtr - SymbolNameStartPtr);
 }
 
 const TCHAR* ParseHLSLSymbolName(const TCHAR* SearchString, FString& SymbolName)
 {
-	FStringView Result = ParseHLSLSymbolName(SearchString);
+	FStringView Result = ParseHLSLSymbolName<TCHAR, FStringView>(SearchString);
 
 	SymbolName = FString(Result);
 
@@ -474,7 +456,7 @@ void UE::ShaderCompilerCommon::ParseRayTracingEntryPoint(const FString& Input, F
 }
 
 
-bool UE::ShaderCompilerCommon::RemoveDeadCode(FString& InOutPreprocessedShaderSource, TConstArrayView<FStringView> RequiredSymbols, TArray<FShaderCompilerError>& OutErrors)
+bool UE::ShaderCompilerCommon::RemoveDeadCode(FShaderSource& InOutPreprocessedShaderSource, TConstArrayView<FStringView> InRequiredSymbols, TArray<FShaderCompilerError>& OutErrors)
 {
 	TRACE_CPUPROFILER_EVENT_SCOPE(RemoveDeadCode);
 
@@ -485,14 +467,14 @@ bool UE::ShaderCompilerCommon::RemoveDeadCode(FString& InOutPreprocessedShaderSo
 	           |  UE::ShaderMinifier::EMinifyShaderFlags::OutputStats;  // Output a comment detailing how many blocks of each type (functions/structs/etc.) were emitted
 #endif
 
-	UE::ShaderMinifier::FMinifiedShader Minified  = UE::ShaderMinifier::Minify(InOutPreprocessedShaderSource, RequiredSymbols,
+	UE::ShaderMinifier::FMinifiedShader Minified = UE::ShaderMinifier::Minify(InOutPreprocessedShaderSource, InRequiredSymbols,
 		  UE::ShaderMinifier::EMinifyShaderFlags::OutputCommentLines // Preserve comments that were left after preprocessing
 		| UE::ShaderMinifier::EMinifyShaderFlags::OutputLines        // Emit #line directives
 		| ExtraFlags);
 
 	if (Minified.Success())
 	{
-		Swap(InOutPreprocessedShaderSource, Minified.Code);
+		InOutPreprocessedShaderSource = MoveTemp(Minified.Code);
 		return true;
 	}
 	else
@@ -502,12 +484,12 @@ bool UE::ShaderCompilerCommon::RemoveDeadCode(FString& InOutPreprocessedShaderSo
 	}
 }
 
-bool UE::ShaderCompilerCommon::RemoveDeadCode(FString& InOutPreprocessedShaderSource, const FString& EntryPoint, TArray<FShaderCompilerError>& OutErrors)
+bool UE::ShaderCompilerCommon::RemoveDeadCode(FShaderSource& InOutPreprocessedShaderSource, const FString& EntryPoint, TArray<FShaderCompilerError>& OutErrors)
 {
 	return UE::ShaderCompilerCommon::RemoveDeadCode(InOutPreprocessedShaderSource, EntryPoint, {}, OutErrors);
 }
 
-bool UE::ShaderCompilerCommon::RemoveDeadCode(FString& InOutPreprocessedShaderSource, const FString& EntryPoint, TConstArrayView<FStringView> InRequiredSymbols, TArray<FShaderCompilerError>& OutErrors)
+bool UE::ShaderCompilerCommon::RemoveDeadCode(FShaderSource& InOutPreprocessedShaderSource, const FString& EntryPoint, TConstArrayView<FStringView> InRequiredSymbols, TArray<FShaderCompilerError>& OutErrors)
 {
 	TArray<FStringView> RequiredSymbols;
 
@@ -785,9 +767,9 @@ struct FUniformBufferInfo
 struct FUniformBufferMemberInfoNew
 {
 	// eg View.WorldToClip
-	FStringView NameAsStructMember;
+	FShaderSource::FViewType NameAsStructMember;
 	// eg View_WorldToClip
-	FStringView GlobalName;
+	FShaderSource::FViewType GlobalName;
 
 	bool operator<(const FUniformBufferMemberInfoNew& Other)
 	{
@@ -811,7 +793,7 @@ struct FUniformBufferMemberView
 
 struct FUniformBufferInfoNew
 {
-	FStringView Name;
+	FShaderSource::FViewType Name;
 	int32 NextWithSameLength;							// Linked list of uniform buffer infos with same name length
 	TArray<FUniformBufferMemberInfoNew> Members;		// Members sorted by length
 	TArray<FUniformBufferMemberView> MembersByLength;	// Offset and count of members of a given length
@@ -826,10 +808,10 @@ struct FUniformBufferSpan
 
 // Compacts spaces out of a compound identifier.  Returns the new end pointer of the compacted identifier.
 // End and result pointers are exclusive (length of the string is End - Start).
-static TCHAR* CompactCompoundIdentifier(TCHAR* Start, TCHAR* End)
+static FShaderSource::CharType* CompactCompoundIdentifier(FShaderSource::CharType* Start, FShaderSource::CharType* End)
 {
 	// Find first whitespace in the identifier, if present
-	TCHAR* ReadChar;
+	FShaderSource::CharType* ReadChar;
 	for (ReadChar = Start; ReadChar < End; ++ReadChar)
 	{
 		if (IsSpaceOrTabOrEOL(*ReadChar))
@@ -845,7 +827,7 @@ static TCHAR* CompactCompoundIdentifier(TCHAR* Start, TCHAR* End)
 
 	// Found some whitespace, so we need to compact the non-whitespace, swapping the whitespace to the end of the range
 	// WriteChar here will be the first whitespace character that we need to compact into.
-	TCHAR* WriteChar = ReadChar;
+	FShaderSource::CharType* WriteChar = ReadChar;
 	for (++ReadChar; ReadChar < End; ++ReadChar)
 	{
 		// If the current read character is non-whitespace, compact it down
@@ -858,32 +840,32 @@ static TCHAR* CompactCompoundIdentifier(TCHAR* Start, TCHAR* End)
 	return WriteChar;
 }
 
-const TCHAR* ParseUniformBufferDefinition(const TCHAR* ReadStart, TArray<FUniformBufferInfoNew>& UniformBufferInfos, uint64 UniformBufferFilter[64], int32 UniformBuffersByLength[64])
+const FShaderSource::CharType* ParseUniformBufferDefinition(const FShaderSource::CharType* ReadStart, TArray<FUniformBufferInfoNew>& UniformBufferInfos, uint64 UniformBufferFilter[64], int32 UniformBuffersByLength[64])
 {
 	// TODO:  should we check for an existing item?  In my testing, there's only one uniform buffer declaration with a given name,
 	// but the original code used a map, theoretically allowing for multiple.
 	int32 InfoIndex = UniformBufferInfos.AddDefaulted();
 	FUniformBufferInfoNew& Info = UniformBufferInfos[InfoIndex];
 
-	Info.Name = ParseHLSLSymbolName(ReadStart);
+	Info.Name = ParseHLSLSymbolName<FShaderSource::CharType, FShaderSource::FViewType>(ReadStart);
 	check(Info.Name.Len() < 64);
 
-	const TCHAR* OpeningBrace = FindNextChar(ReadStart, '{');
-	const TCHAR* ClosingBrace = FindMatchingClosingBrace(OpeningBrace + 1);
+	const FShaderSource::CharType* OpeningBrace = FindNextChar(ReadStart, '{');
+	const FShaderSource::CharType* ClosingBrace = FindMatchingClosingBrace(OpeningBrace + 1);
 
-	const TCHAR* CurrentParseStart = OpeningBrace + 1;
-	const TCHAR* NextSemicolon = FindNextChar(CurrentParseStart, ';');
+	const FShaderSource::CharType* CurrentParseStart = OpeningBrace + 1;
+	const FShaderSource::CharType* NextSemicolon = FindNextChar(CurrentParseStart, ';');
 
 	while (NextSemicolon < ClosingBrace)
 	{
-		const TCHAR* NextSeparator = FindNextChar(CurrentParseStart, '=');
+		const FShaderSource::CharType* NextSeparator = FindNextChar(CurrentParseStart, '=');
 		if (NextSeparator < NextSemicolon)
 		{
-			const TCHAR* StructStart = CurrentParseStart;
-			const TCHAR* StructEnd = NextSeparator - 1;
+			const FShaderSource::CharType* StructStart = CurrentParseStart;
+			const FShaderSource::CharType* StructEnd = NextSeparator - 1;
 
-			const TCHAR* GlobalStart = NextSeparator + 1;
-			const TCHAR* GlobalEnd = NextSemicolon - 1;
+			const FShaderSource::CharType* GlobalStart = NextSeparator + 1;
+			const FShaderSource::CharType* GlobalEnd = NextSemicolon - 1;
 
 			while (IsSpaceOrTabOrEOL(*StructStart))
 			{
@@ -894,14 +876,14 @@ const TCHAR* ParseUniformBufferDefinition(const TCHAR* ReadStart, TArray<FUnifor
 				GlobalStart++;
 			}
 
-			StructEnd = CompactCompoundIdentifier(const_cast<TCHAR*>(StructStart), const_cast<TCHAR*>(StructEnd));
-			GlobalEnd = CompactCompoundIdentifier(const_cast<TCHAR*>(GlobalStart), const_cast<TCHAR*>(GlobalEnd));
+			StructEnd = CompactCompoundIdentifier(const_cast<FShaderSource::CharType*>(StructStart), const_cast<FShaderSource::CharType*>(StructEnd));
+			GlobalEnd = CompactCompoundIdentifier(const_cast<FShaderSource::CharType*>(GlobalStart), const_cast<FShaderSource::CharType*>(GlobalEnd));
 
-			FStringView StructName(StructStart, StructEnd - StructStart);
-			FStringView GlobalName(GlobalStart, GlobalEnd - GlobalStart);
+			FShaderSource::FViewType StructName(StructStart, StructEnd - StructStart);
+			FShaderSource::FViewType GlobalName(GlobalStart, GlobalEnd - GlobalStart);
 
 			// Avoid unnecessary conversions
-			if (StructName.Len() == GlobalName.Len() && FCString::Strncmp(StructName.GetData(), GlobalName.GetData(), StructName.Len()) != 0)
+			if (StructName.Len() == GlobalName.Len() && FShaderSource::FCStringType::Strncmp(StructName.GetData(), GlobalName.GetData(), StructName.Len()) != 0)
 			{
 				FUniformBufferMemberInfoNew NewMemberInfo;
 				NewMemberInfo.NameAsStructMember = StructName;
@@ -918,7 +900,7 @@ const TCHAR* ParseUniformBufferDefinition(const TCHAR* ReadStart, TArray<FUnifor
 		NextSemicolon = FindNextChar(CurrentParseStart, ';');
 	}
 
-	const TCHAR* EndPtr = ClosingBrace;
+	const FShaderSource::CharType* EndPtr = ClosingBrace;
 
 	// Skip to the end of the UniformBuffer
 	while (*EndPtr && *EndPtr != ';')
@@ -994,9 +976,9 @@ static uint8 AsciiFlagTable[256] =
 
 struct FCompoundIdentifierResult
 {
-	const TCHAR* Identifier;			// Start of identifier
-	const TCHAR* IdentifierEnd;			// End of entire identifier
-	const TCHAR* IdentifierRootEnd;		// End of root token of identifier
+	const FShaderSource::CharType* Identifier;			// Start of identifier
+	const FShaderSource::CharType* IdentifierEnd;			// End of entire identifier
+	const FShaderSource::CharType* IdentifierRootEnd;		// End of root token of identifier
 };
 
 // Searches for a "compound identifier" (series of symbol tokens separated by dots) that also passes the "RootIdentifierFilter".
@@ -1004,9 +986,9 @@ struct FCompoundIdentifierResult
 // with letters or underscore, we can store a 64-bit mask representing ASCII characters 64..127, as all valid start characters are
 // in that range.  As an example, if "View" is a valid root identifier, RootIdentifierFilter[4] will have the bit ('V' - 64) set,
 // and any other 4 character identifier that doesn't start with that letter can be skipped, saving overhead in the caller.
-bool FindNextCompoundIdentifier(const TCHAR*& Search, const uint64 RootIdentifierFilter[64], FCompoundIdentifierResult& OutResult)
+bool FindNextCompoundIdentifier(const FShaderSource::CharType*& Search, const uint64 RootIdentifierFilter[64], FCompoundIdentifierResult& OutResult)
 {
-	const TCHAR* SearchChar = Search;
+	const FShaderSource::CharType* SearchChar = Search;
 	uint8 SearchCharFlag = AsciiFlagTable[(uint8)*SearchChar];
 
 	// Scanning loop
@@ -1139,7 +1121,7 @@ bool FindNextCompoundIdentifier(const TCHAR*& Search, const uint64 RootIdentifie
 		{
 			// Quote, skip to next Quote (or maybe end of string if text is malformed), ignoring the quote if it's escaped
 			SearchChar++;
-			while (*SearchChar && (*SearchChar != TEXT('\"') || *(SearchChar - 1) == TEXT('\\')))
+			while (*SearchChar && (*SearchChar != '\"' || *(SearchChar - 1) == '\\'))
 			{
 				SearchChar++;
 			}
@@ -1152,12 +1134,12 @@ bool FindNextCompoundIdentifier(const TCHAR*& Search, const uint64 RootIdentifie
 			SearchCharFlag = AsciiFlagTable[(uint8)*SearchChar];
 		}
 		// Must be null terminator or slash at this point -- we've tested all other possibilities
-		else if (*SearchChar == TEXT('/'))
+		else if (*SearchChar == '/')
 		{
 			// Check if this is a commented out block (typically a commented out uniform declaration) and skip over it.
 			// If the text is bad, there could be a /* right at the end of the string, so we need to check there is at least
 			// one more character.
-			if (SearchChar[1] == TEXT('*') && SearchChar[2] != 0)
+			if (SearchChar[1] == '*' && SearchChar[2] != 0)
 			{
 				// Search for slash (or end of string), starting at SearchChar + 3.  If we find a slash, we'll check the previous
 				// character to see if it's the end of the comment.  Starting at +3 is necessary to avoid matching a slash as the
@@ -1172,7 +1154,7 @@ bool FindNextCompoundIdentifier(const TCHAR*& Search, const uint64 RootIdentifie
 					}
 
 					// Is this the end of the comment?
-					if (*(SearchChar - 1) == TEXT('*'))
+					if (*(SearchChar - 1) == '*')
 					{
 						if (*SearchChar)
 						{
@@ -1211,15 +1193,15 @@ bool FindNextCompoundIdentifier(const TCHAR*& Search, const uint64 RootIdentifie
 	}
 }
 
-TCHAR* FindNextUniformBufferDefinition(TCHAR* SearchPtr, TCHAR* SourceStart, const TCHAR* UniformBufferStructIdentifier, int32 StructIdentifierLen)
+FShaderSource::CharType* FindNextUniformBufferDefinition(FShaderSource::CharType* SearchPtr, FShaderSource::CharType* SourceStart, FShaderSource::FViewType UniformBufferStructIdentifier)
 {
 	while (SearchPtr)
 	{
-		SearchPtr = FCString::Strstr(SearchPtr, UniformBufferStructIdentifier);
+		SearchPtr = FShaderSource::FCStringType::Strstr(SearchPtr, UniformBufferStructIdentifier.GetData());
 
 		if (SearchPtr)
 		{
-			if (SearchPtr > SourceStart && IsSpaceOrTabOrEOL(*(SearchPtr - 1)) && IsSpaceOrTabOrEOL(*(SearchPtr + StructIdentifierLen)))
+			if (SearchPtr > SourceStart && IsSpaceOrTabOrEOL(*(SearchPtr - 1)) && IsSpaceOrTabOrEOL(*(SearchPtr + UniformBufferStructIdentifier.Len())))
 			{
 				break;
 			}
@@ -1232,9 +1214,9 @@ TCHAR* FindNextUniformBufferDefinition(TCHAR* SearchPtr, TCHAR* SourceStart, con
 	return SearchPtr;
 }
 
-const TCHAR* FindPreviousDot(const TCHAR* SearchPtr, const TCHAR* SearchMin)
+const FShaderSource::CharType* FindPreviousDot(const FShaderSource::CharType* SearchPtr, const FShaderSource::CharType* SearchMin)
 {
-	while ((SearchPtr > SearchMin) && (*SearchPtr != TEXT('.')))
+	while ((SearchPtr > SearchMin) && (*SearchPtr != '.'))
 	{
 		SearchPtr--;
 	}
@@ -1243,7 +1225,7 @@ const TCHAR* FindPreviousDot(const TCHAR* SearchPtr, const TCHAR* SearchMin)
 
 // The cross compiler doesn't yet support struct initializers needed to construct static structs for uniform buffers
 // Replace all uniform buffer struct member references (View.WorldToClip) with a flattened name that removes the struct dependency (View_WorldToClip)
-void CleanupUniformBufferCode(const FShaderCompilerEnvironment& Environment, FString& PreprocessedShaderSource)
+void CleanupUniformBufferCode(const FShaderCompilerEnvironment& Environment, FShaderSource& PreprocessedShaderSource)
 {
 	TRACE_CPUPROFILER_EVENT_SCOPE(CleanupUniformBufferCode);
 
@@ -1256,18 +1238,17 @@ void CleanupUniformBufferCode(const FShaderCompilerEnvironment& Environment, FSt
 	UniformBufferSpans.Reserve(Environment.UniformBufferMap.Num());
 	memset(UniformBuffersByLength, 0xff, sizeof(UniformBuffersByLength));
 
-	const TCHAR* UniformBufferStructIdentifier = TEXT("UniformBuffer");
-	const int32 StructIdentifierLen = FCString::Strlen(UniformBufferStructIdentifier);
+	FShaderSource::FViewType UniformBufferStructIdentifier = SHADER_SOURCE_VIEWLITERAL("UniformBuffer");
 
-	TCHAR* SourceStart = &PreprocessedShaderSource[0];
-	TCHAR* SearchPtr = SourceStart;
-	TCHAR* EndOfPreviousUniformBuffer = SourceStart;
+	FShaderSource::CharType* SourceStart = PreprocessedShaderSource.GetData();
+	FShaderSource::CharType* SearchPtr = SourceStart;
+	FShaderSource::CharType* EndOfPreviousUniformBuffer = SourceStart;
 	bool bUniformBufferFound;
 
 	do
 	{
 		// Find the next uniform buffer definition
-		SearchPtr = FindNextUniformBufferDefinition(SearchPtr, SourceStart, UniformBufferStructIdentifier, StructIdentifierLen);
+		SearchPtr = FindNextUniformBufferDefinition(SearchPtr, SourceStart, UniformBufferStructIdentifier);
 
 		if (SearchPtr)
 		{
@@ -1284,27 +1265,27 @@ void CleanupUniformBufferCode(const FShaderCompilerEnvironment& Environment, FSt
 		// were found).  If there are no uniform buffers yet, we don't need to parse anything.
 		if (UniformBufferInfos.Num())
 		{
-			const TCHAR* ParsePtr = EndOfPreviousUniformBuffer;
+			const FShaderSource::CharType* ParsePtr = EndOfPreviousUniformBuffer;
 
 			FCompoundIdentifierResult Result;
 			while (FindNextCompoundIdentifier(ParsePtr, UniformBufferFilter, Result))
 			{
 				// Check if the identifier corresponds to a uniform buffer
-				FStringView IdentifierRoot(Result.Identifier, Result.IdentifierRootEnd - Result.Identifier);
+				FShaderSource::FViewType IdentifierRoot(Result.Identifier, Result.IdentifierRootEnd - Result.Identifier);
 				for (int32 UniformInfoIndex = UniformBuffersByLength[IdentifierRoot.Len()]; UniformInfoIndex != INDEX_NONE; UniformInfoIndex = UniformBufferInfos[UniformInfoIndex].NextWithSameLength)
 				{
 					FUniformBufferInfoNew& Info = UniformBufferInfos[UniformInfoIndex];
 					if (IdentifierRoot.Equals(Info.Name, ESearchCase::CaseSensitive))
 					{
 						// Found the uniform buffer, clean up potential whitespace
-						Result.IdentifierEnd = CompactCompoundIdentifier(const_cast<TCHAR*>(Result.Identifier), const_cast<TCHAR*>(Result.IdentifierEnd));
+						Result.IdentifierEnd = CompactCompoundIdentifier(const_cast<FShaderSource::CharType*>(Result.Identifier), const_cast<FShaderSource::CharType*>(Result.IdentifierEnd));
 
 						// Now try to find a matching member.  We need to check subsets of the full "identifier", to strip away function calls, components, or child structures.
 						bool bMatchFound = false;
 
 						for (; Result.IdentifierEnd > Result.IdentifierRootEnd; Result.IdentifierEnd = FindPreviousDot(Result.IdentifierEnd - 1, Result.IdentifierRootEnd))
 						{
-							FStringView Identifier(Result.Identifier, Result.IdentifierEnd - Result.Identifier);
+							FShaderSource::FViewType Identifier(Result.Identifier, Result.IdentifierEnd - Result.Identifier);
 							if (Identifier.Len() < Info.MembersByLength.Num())
 							{
 								const FUniformBufferMemberView& MemberView = Info.MembersByLength[Identifier.Len()];
@@ -1318,8 +1299,8 @@ void CleanupUniformBufferCode(const FShaderCompilerEnvironment& Environment, FSt
 										const int32 OriginalTextLen = Info.Members[MemberIndex].NameAsStructMember.Len();
 										const int32 ReplacementTextLen = Info.Members[MemberIndex].GlobalName.Len();
 
-										const TCHAR* GlobalNameStart = GetData(Info.Members[MemberIndex].GlobalName);
-										TCHAR* IdentifierStart = const_cast<TCHAR*>(Result.Identifier);
+										const FShaderSource::CharType* GlobalNameStart = GetData(Info.Members[MemberIndex].GlobalName);
+										FShaderSource::CharType* IdentifierStart = const_cast<FShaderSource::CharType*>(Result.Identifier);
 
 										int32 Index = 0;
 										for (; Index < ReplacementTextLen; Index++)
@@ -1353,8 +1334,8 @@ void CleanupUniformBufferCode(const FShaderCompilerEnvironment& Environment, FSt
 			// Unterminate the string (put the first character of the struct identifier back in place) and parse it
 			*SearchPtr = UniformBufferStructIdentifier[0];
 
-			const TCHAR* ConstStructEndPtr = ParseUniformBufferDefinition(SearchPtr + StructIdentifierLen, UniformBufferInfos, UniformBufferFilter, UniformBuffersByLength);
-			TCHAR* StructEndPtr = &PreprocessedShaderSource[ConstStructEndPtr - &PreprocessedShaderSource[0]];
+			const FShaderSource::CharType* ConstStructEndPtr = ParseUniformBufferDefinition(SearchPtr + UniformBufferStructIdentifier.Len(), UniformBufferInfos, UniformBufferFilter, UniformBuffersByLength);
+			FShaderSource::CharType* StructEndPtr = &SourceStart[ConstStructEndPtr - &SourceStart[0]];
 
 			// Comment out the uniform buffer struct and initializer
 			*SearchPtr = '/';
@@ -1374,21 +1355,21 @@ void CleanupUniformBufferCode(const FShaderCompilerEnvironment& Environment, FSt
 	// it's not necessary to add a line directive to fix up line numbers because uniform buffer declarations are always in generated files, and there
 	// will be a line directive already there for the transition from the generated file back to whatever file included it.  The destination offset
 	// for the first move is the start of the first uniform buffer declaration we are overwriting, then advances as characters are copied.
-	int32 DestOffset = UniformBufferSpans.Num() ? UniformBufferSpans[0].Offset : PreprocessedShaderSource.GetCharArray().Num();
+	int32 DestOffset = UniformBufferSpans.Num() ? UniformBufferSpans[0].Offset : PreprocessedShaderSource.Len();
 
 	for (int32 SpanIndex = 0; SpanIndex < UniformBufferSpans.Num(); SpanIndex++)
 	{
-		// The source code we are compacting down is from the end of one span to the start of the next span, or end of the string (plus one so it also
-		// moves the null terminator).
+		// The source code we are compacting down is from the end of one span to the start of the next span, or end of the string.
+		// We do not need to account for null terminator as the ShrinkToLen call below will null terminate for us.
 		int32 SourceOffset = UniformBufferSpans[SpanIndex].Offset + UniformBufferSpans[SpanIndex].Length;
-		int32 MoveCount = (SpanIndex < UniformBufferSpans.Num() - 1 ? UniformBufferSpans[SpanIndex + 1].Offset : PreprocessedShaderSource.Len() + 1) - SourceOffset;
+		int32 MoveCount = (SpanIndex < UniformBufferSpans.Num() - 1 ? UniformBufferSpans[SpanIndex + 1].Offset : PreprocessedShaderSource.Len()) - SourceOffset;
 
-		check(DestOffset >= 0 && DestOffset < SourceOffset && SourceOffset + MoveCount <= PreprocessedShaderSource.GetCharArray().Num());
+		check(DestOffset >= 0 && DestOffset < SourceOffset && SourceOffset + MoveCount <= PreprocessedShaderSource.Len());
 
-		memmove(SourceStart + DestOffset, SourceStart + SourceOffset, MoveCount * sizeof(TCHAR));
+		memmove(SourceStart + DestOffset, SourceStart + SourceOffset, MoveCount * sizeof(FShaderSource::CharType));
 		DestOffset += MoveCount;
 	}
-	PreprocessedShaderSource.GetCharArray().SetNum(DestOffset, false);
+	PreprocessedShaderSource.ShrinkToLen(DestOffset, false);
 }
 
 ///////////////////////////////////////////////////////////////////////////////////////////////////////////////////////////////////////
@@ -1708,24 +1689,23 @@ namespace UE::ShaderCompilerCommon
 		)
 	{
 		TRACE_CPUPROFILER_EVENT_SCOPE(FBaseShaderFormat_PreprocessShader);
-		PreprocessOutput.EditSource().Empty();
 
 		if (EnumHasAnyFlags(Input.DebugInfoFlags, EShaderDebugInfoFlags::CompileFromDebugUSF))
 		{
 			// the "VirtualSourceFilePath" given is actually an absolute path to a dumped debug USF file; load it directly.
 			// this occurs when running SCW in "direct compile" mode; this file will already be preprocessed.
-			bool bSuccess = FFileHelper::LoadFileToString(PreprocessOutput.EditSource(), *Input.VirtualSourceFilePath);
+			FString DebugUSF;
+			bool bSuccess = FFileHelper::LoadFileToString(DebugUSF, *Input.VirtualSourceFilePath);
 
 			if (bSuccess)
 			{
 				// const_cast for compile environment; need to populate a subset of environment parameters from parsing comments in the preprocessed code
-				UE::ShaderCompilerCommon::SerializeEnvironmentFromBase64(const_cast<FShaderCompilerEnvironment&>(Input.Environment), PreprocessOutput.GetSource());
+				UE::ShaderCompilerCommon::SerializeEnvironmentFromBase64(const_cast<FShaderCompilerEnvironment&>(Input.Environment), DebugUSF);
 
 				// strip comments from source when loading from a debug USF. some backends don't handle the comments that the debug dump inserts properly.
-				// this (currently) incurs an extra conversion (TCHAR -> ANSI -> TCHAR), but since this is only used in the debug path the perf hit is irrelevant
 				TArray<ANSICHAR> Stripped;
-				ShaderConvertAndStripComments(PreprocessOutput.GetSource(), Stripped);
-				PreprocessOutput.EditSource() = Stripped.GetData();
+				ShaderConvertAndStripComments(DebugUSF, Stripped);
+				PreprocessOutput.EditSource().Set({ Stripped.GetData(), Stripped.Num() });
 			}
 
 			return bSuccess;
@@ -1769,7 +1749,7 @@ namespace UE::ShaderCompilerCommon
 		DumpExtendedDebugShaderData(Input, PreprocessOutput, Output);
 	}
 
-	void DumpDebugShaderData(const FShaderCompilerInput& Input, const FString& PreprocessedSource, const FDebugShaderDataOptions& Options)
+	void DumpDebugShaderData(const FShaderCompilerInput& Input, FStringView PreprocessedSource, const FDebugShaderDataOptions& Options)
 	{
 		if (!Input.DumpDebugInfoEnabled())
 		{
@@ -1804,14 +1784,19 @@ namespace UE::ShaderCompilerCommon
 			}
 			
 			PrefixedOptions.FilenamePrefix = *StrippedPrefix;
-			FFileHelper::SaveStringToFile(PreprocessOutput.GetSource(), *PrefixedOptions.GetDebugShaderPath(Input));
+			FFileHelper::SaveStringToFile(PreprocessOutput.GetSourceViewWide(), *PrefixedOptions.GetDebugShaderPath(Input));
 
 			PrefixedOptions.FilenamePrefix = *PreprocessedPrefix;
-			FFileHelper::SaveStringToFile(PreprocessOutput.GetUnstrippedSource(), *PrefixedOptions.GetDebugShaderPath(Input));
+			FFileHelper::SaveStringToFile(PreprocessOutput.GetUnstrippedSourceView(), *PrefixedOptions.GetDebugShaderPath(Input));
+		}
+		if (Output.ModifiedShaderSource.IsEmpty())
+		{
+			DumpDebugShaderData(Input, PreprocessOutput.GetSourceViewWide(), Options);
+		}
+		else
+		{
+			DumpDebugShaderData(Input, FStringView(Output.ModifiedShaderSource), Options);
 		}
-			
-		const FString& SourceToDump = Output.ModifiedShaderSource.IsEmpty() ? PreprocessOutput.GetSource() : Output.ModifiedShaderSource;
-		DumpDebugShaderData(Input, SourceToDump, Options);
 		FFileHelper::SaveStringToFile(Output.OutputHash.ToString(), *GetDebugFileName(Input, Options, TEXT("OutputHash.txt")), FFileHelper::EEncodingOptions::ForceAnsi);
 
 		if (EnumHasAnyFlags(Input.DebugInfoFlags, EShaderDebugInfoFlags::Diagnostics))
@@ -1888,7 +1873,7 @@ namespace UE::ShaderCompilerCommon
 		Env.SerializeCompilationDependencies(Ar);
 	}
 
-	FString GetDebugShaderContents(const FShaderCompilerInput& Input, const FString& PreprocessedSource, const FDebugShaderDataOptions& Options)
+	FString GetDebugShaderContents(const FShaderCompilerInput& Input, FStringView PreprocessedSource, const FDebugShaderDataOptions& Options)
 	{
 		// If preprocessed cache is enabled, debug dump occurs in the cook process rather than the workers, and
 		// in that case the env in Input.Environment has not been merged with the shared env. Do so here.
@@ -1898,7 +1883,7 @@ namespace UE::ShaderCompilerCommon
 			MergedEnvironment.Merge(*Input.SharedEnvironment);
 		}
 
-		FString Contents = MergedEnvironment.GetDefinitionsAsCommentedCode();
+		FString Contents = Options.AppendPreSource ? Options.AppendPreSource() : FString();
 
 		if (Options.AppendPreSource)
 		{
diff --git a/Engine/Source/Developer/ShaderCompilerCommon/Private/ShaderConductorContext.cpp b/Engine/Source/Developer/ShaderCompilerCommon/Private/ShaderConductorContext.cpp
--- a/Engine/Source/Developer/ShaderCompilerCommon/Private/ShaderConductorContext.cpp
+++ b/Engine/Source/Developer/ShaderCompilerCommon/Private/ShaderConductorContext.cpp
@@ -599,11 +599,11 @@ namespace CrossCompiler
 	}
 
 	PRAGMA_DISABLE_DEPRECATION_WARNINGS		// FShaderCompilerDefinitions will be made internal in the future, marked deprecated until then
-	bool FShaderConductorContext::LoadSource(const FString& ShaderSource, const FString& Filename, const FString& EntryPoint, EShaderFrequency ShaderStage, const FShaderCompilerDefinitions* Definitions, const TArray<FString>* ExtraDxcArgs)
+	bool FShaderConductorContext::LoadSource(FStringView ShaderSource, const FString& Filename, const FString& EntryPoint, EShaderFrequency ShaderStage, const FShaderCompilerDefinitions* Definitions, const TArray<FString>* ExtraDxcArgs)
 	PRAGMA_ENABLE_DEPRECATION_WARNINGS
 	{
 		// Convert FString to ANSI string and store them as intermediates
-		ConvertFStringToAnsiString(ShaderSource, Intermediates->ShaderSource);
+		ConvertFStringViewToAnsiString(ShaderSource, Intermediates->ShaderSource);
 		ConvertFStringToAnsiString(Filename, Intermediates->Filename);
 		ConvertFStringToAnsiString(EntryPoint, Intermediates->EntryPoint);
 
@@ -624,6 +624,13 @@ namespace CrossCompiler
 		return true;
 	}
 
+	PRAGMA_DISABLE_DEPRECATION_WARNINGS		// FShaderCompilerDefinitions will be made internal in the future, marked deprecated until then
+	bool FShaderConductorContext::LoadSource(const FString& ShaderSource, const FString& Filename, const FString& EntryPoint, EShaderFrequency ShaderStage, const FShaderCompilerDefinitions* Definitions, const TArray<FString>* ExtraDxcArgs)
+	PRAGMA_ENABLE_DEPRECATION_WARNINGS
+	{
+		return LoadSource(FStringView(ShaderSource), Filename, EntryPoint, ShaderStage, Definitions, ExtraDxcArgs);
+	}
+
 	PRAGMA_DISABLE_DEPRECATION_WARNINGS		// FShaderCompilerDefinitions will be made internal in the future, marked deprecated until then
 	bool FShaderConductorContext::LoadSource(const ANSICHAR* ShaderSource, const ANSICHAR* Filename, const ANSICHAR* EntryPoint, EShaderFrequency ShaderStage, const FShaderCompilerDefinitions* Definitions, const TArray<FString>* ExtraDxcArgs)
 	PRAGMA_ENABLE_DEPRECATION_WARNINGS
@@ -969,6 +976,11 @@ namespace CrossCompiler
 		return false; // Dummy
 	}
 
+	bool FShaderConductorContext::LoadSource(FStringView ShaderSource, const FString& Filename, const FString& EntryPoint, EShaderFrequency ShaderStage, const FShaderCompilerDefinitions* Definitions, const TArray<FString>* ExtraDxcArgs)
+	{
+		return false; // Dummy
+	}
+
 	bool FShaderConductorContext::LoadSource(const ANSICHAR* ShaderSource, const ANSICHAR* Filename, const ANSICHAR* EntryPoint, EShaderFrequency ShaderStage, const FShaderCompilerDefinitions* Definitions, const TArray<FString>* ExtraDxcArgs)
 	{
 		return false; // Dummy
diff --git a/Engine/Source/Developer/ShaderCompilerCommon/Private/ShaderMinifier.cpp b/Engine/Source/Developer/ShaderCompilerCommon/Private/ShaderMinifier.cpp
--- a/Engine/Source/Developer/ShaderCompilerCommon/Private/ShaderMinifier.cpp
+++ b/Engine/Source/Developer/ShaderCompilerCommon/Private/ShaderMinifier.cpp
@@ -24,45 +24,45 @@ DEFINE_LOG_CATEGORY_STATIC(LogShaderMinifier, Log, All);
 namespace UE::ShaderMinifier
 {
 
-static FStringView SubStrView(FStringView S, int32 Start)
+static FShaderSource::FViewType SubStrView(FShaderSource::FViewType S, int32 Start)
 {
 	Start = FMath::Min(Start, S.Len());
 	int32 Len = S.Len() - Start;
-	return FStringView(S.GetData() + Start, Len);
+	return FShaderSource::FViewType(S.GetData() + Start, Len);
 }
 
-static FStringView SubStrView(FStringView S, int32 Start, int32 Len)
+static FShaderSource::FViewType SubStrView(FShaderSource::FViewType S, int32 Start, int32 Len)
 {
 	Start = FMath::Min(Start, S.Len());
 	Len	  = FMath::Min(Len, S.Len() - Start);
-	return FStringView(S.GetData() + Start, Len);
+	return FShaderSource::FViewType(S.GetData() + Start, Len);
 }
 
 template<typename TCondition>
-static FStringView SkipUntil(FStringView Source, TCondition Cond)
+static FShaderSource::FViewType SkipUntil(FShaderSource::FViewType Source, TCondition Cond)
 {
 	int32 Cursor = 0;
 	const int32 SourceLen = Source.Len();
 	while (Cursor < SourceLen)
 	{
-		if (Cond(FStringView(Source.GetData() + Cursor, SourceLen - Cursor)))
+		if (Cond(FShaderSource::FViewType(Source.GetData() + Cursor, SourceLen - Cursor)))
 		{
 			break;
 		}
 		++Cursor;
 	}
-	return FStringView(Source.GetData() + Cursor, SourceLen - Cursor);
+	return FShaderSource::FViewType(Source.GetData() + Cursor, SourceLen - Cursor);
 }
 
-static bool Equals(FStringView A, FStringView B)
+static bool Equals(FShaderSource::FViewType A, FShaderSource::FViewType B)
 {
 	int32 Len = A.Len();
 	if (Len != B.Len())
 	{
 		return false;
 	}
-	const TCHAR* DataA = A.GetData();
-	const TCHAR* DataB = B.GetData();
+	const FShaderSource::CharType* DataA = A.GetData();
+	const FShaderSource::CharType* DataB = B.GetData();
 	for (int32 I = 0; I < Len; ++I)
 	{
 		if (DataA[I] != DataB[I])
@@ -73,7 +73,7 @@ static bool Equals(FStringView A, FStringView B)
 	return true;
 }
 
-static bool StartsWith(FStringView Source, FStringView Prefix)
+static bool StartsWith(FShaderSource::FViewType Source, FShaderSource::FViewType Prefix)
 {
 	const int32 SourceLen = Source.Len();
 	const int32 PrefixLen = Prefix.Len();
@@ -81,7 +81,7 @@ static bool StartsWith(FStringView Source, FStringView Prefix)
 	{
 		return false;
 	}
-	FStringView SourceView(Source.GetData(), PrefixLen);
+	FShaderSource::FViewType SourceView(Source.GetData(), PrefixLen);
 	return Equals(SourceView, Prefix);
 }
 
@@ -132,22 +132,22 @@ struct FCharacterFlags
 		Flags[uint8(')')]  |= Special;
 	}
 
-	bool IsSpace(TCHAR C) const
+	bool IsSpace(FShaderSource::CharType C) const
 	{
 		return (Flags[uint8(C)] & Space) != 0;
 	}
 
-	bool IsNumber(TCHAR C) const
+	bool IsNumber(FShaderSource::CharType C) const
 	{
 		return (Flags[uint8(C)] & Number) != 0;
 	}
 
-	bool IsPossibleIdentifierCharacter(TCHAR C) const
+	bool IsPossibleIdentifierCharacter(FShaderSource::CharType C) const
 	{
 		return (Flags[uint8(C)] & PossibleIdentifierMask) != 0;
 	}
 
-	bool IsSpecial(TCHAR C) const
+	bool IsSpecial(FShaderSource::CharType C) const
 	{
 		return (Flags[uint8(C)] & Special) != 0;
 	}
@@ -157,24 +157,24 @@ struct FCharacterFlags
 
 static const FCharacterFlags GCharacterFlags;
 
-static bool IsSpace(TCHAR C)
+static bool IsSpace(FShaderSource::CharType C)
 {
 	return GCharacterFlags.IsSpace(C);
 }
 
-static bool IsNumber(TCHAR C)
+static bool IsNumber(FShaderSource::CharType C)
 {
 	return GCharacterFlags.IsNumber(C);
 }
 
-static bool IsPossibleIdentifierCharacter(TCHAR C)
+static bool IsPossibleIdentifierCharacter(FShaderSource::CharType C)
 {
 	return GCharacterFlags.IsPossibleIdentifierCharacter(C);
 }
 
-static FStringView ExtractOperator(FStringView Source)
+static FShaderSource::FViewType ExtractOperator(FShaderSource::FViewType Source)
 {
-	FStringView Result;
+	FShaderSource::FViewType Result;
 
 	if (Source.IsEmpty() || IsPossibleIdentifierCharacter(Source[0]))
 	{
@@ -182,53 +182,53 @@ static FStringView ExtractOperator(FStringView Source)
 	}
 
 	// NOTE: array is sorted by length to match complete operator character sequences first
-	static const FStringView SupportedOperators[] =
+	static const FShaderSource::FViewType SupportedOperators[] =
 	{
 		// three-character operators
-		FStringView(TEXT("<<=")),
-		FStringView(TEXT(">>=")),
-		FStringView(TEXT("->*")),
+		SHADER_SOURCE_VIEWLITERAL("<<="),
+		SHADER_SOURCE_VIEWLITERAL(">>="),
+		SHADER_SOURCE_VIEWLITERAL("->*"),
 
 		// two-character operators
-		FStringView(TEXT("+=")),
-		FStringView(TEXT("++")),
-		FStringView(TEXT("-=")),
-		FStringView(TEXT("--")),
-		FStringView(TEXT("->")),
-		FStringView(TEXT("*=")),
-		FStringView(TEXT("/=")),
-		FStringView(TEXT("%=")),
-		FStringView(TEXT("^=")),
-		FStringView(TEXT("&=")),
-		FStringView(TEXT("&&")),
-		FStringView(TEXT("|=")),
-		FStringView(TEXT("||")),
-		FStringView(TEXT("<<")),
-		FStringView(TEXT("<=")),
-		FStringView(TEXT(">>")),
-		FStringView(TEXT(">=")),
-		FStringView(TEXT("==")),
-		FStringView(TEXT("!=")),
-		FStringView(TEXT("()")),
-		FStringView(TEXT("[]")),
+		SHADER_SOURCE_VIEWLITERAL("+="),
+		SHADER_SOURCE_VIEWLITERAL("++"),
+		SHADER_SOURCE_VIEWLITERAL("-="),
+		SHADER_SOURCE_VIEWLITERAL("--"),
+		SHADER_SOURCE_VIEWLITERAL("->"),
+		SHADER_SOURCE_VIEWLITERAL("*="),
+		SHADER_SOURCE_VIEWLITERAL("/="),
+		SHADER_SOURCE_VIEWLITERAL("%="),
+		SHADER_SOURCE_VIEWLITERAL("^="),
+		SHADER_SOURCE_VIEWLITERAL("&="),
+		SHADER_SOURCE_VIEWLITERAL("&&"),
+		SHADER_SOURCE_VIEWLITERAL("|="),
+		SHADER_SOURCE_VIEWLITERAL("||"),
+		SHADER_SOURCE_VIEWLITERAL("<<"),
+		SHADER_SOURCE_VIEWLITERAL("<="),
+		SHADER_SOURCE_VIEWLITERAL(">>"),
+		SHADER_SOURCE_VIEWLITERAL(">="),
+		SHADER_SOURCE_VIEWLITERAL("=="),
+		SHADER_SOURCE_VIEWLITERAL("!="),
+		SHADER_SOURCE_VIEWLITERAL("()"),
+		SHADER_SOURCE_VIEWLITERAL("[]"),
 
 		// single character operators
-		FStringView(TEXT("+")),
-		FStringView(TEXT("-")),
-		FStringView(TEXT("*")),
-		FStringView(TEXT("/")),
-		FStringView(TEXT("%")),
-		FStringView(TEXT("^")),
-		FStringView(TEXT("&")),
-		FStringView(TEXT("|")),
-		FStringView(TEXT("~")),
-		FStringView(TEXT("!")),
-		FStringView(TEXT("=")),
-		FStringView(TEXT("<")),
-		FStringView(TEXT(">")),
+		SHADER_SOURCE_VIEWLITERAL("+"),
+		SHADER_SOURCE_VIEWLITERAL("-"),
+		SHADER_SOURCE_VIEWLITERAL("*"),
+		SHADER_SOURCE_VIEWLITERAL("/"),
+		SHADER_SOURCE_VIEWLITERAL("%"),
+		SHADER_SOURCE_VIEWLITERAL("^"),
+		SHADER_SOURCE_VIEWLITERAL("&"),
+		SHADER_SOURCE_VIEWLITERAL("|"),
+		SHADER_SOURCE_VIEWLITERAL("~"),
+		SHADER_SOURCE_VIEWLITERAL("!"),
+		SHADER_SOURCE_VIEWLITERAL("="),
+		SHADER_SOURCE_VIEWLITERAL("<"),
+		SHADER_SOURCE_VIEWLITERAL(">"),
 	};
 
-	for (FStringView Operator : SupportedOperators)
+	for (FShaderSource::FViewType Operator : SupportedOperators)
 	{
 		if (StartsWith(Source, Operator))
 		{
@@ -240,17 +240,17 @@ static FStringView ExtractOperator(FStringView Source)
 	return Result;
 }
 
-static FStringView SkipUntilNonIdentifierCharacter(FStringView Source)
+static FShaderSource::FViewType SkipUntilNonIdentifierCharacter(FShaderSource::FViewType Source)
 {
 	const int32 SourceLen = Source.Len();
 	int32 Cursor = 0;
-	const TCHAR* SourceData = Source.GetData();
+	const FShaderSource::CharType* SourceData = Source.GetData();
 
 #if UE_SHADER_MINIFIER_SSE
+	if constexpr (sizeof(FShaderSource::CharType) == 2)
 	{
-		const int32 AlignedLen = SourceLen & (~7); // align down to multiple of 8 TCHAR-s
 		const __m128i NeedleVec = _mm_setr_epi16(L'0', L'9', L'a', L'z', L'A', L'Z', L'_', L'_');
-		while (Cursor < AlignedLen)
+		while (Cursor < SourceLen)
 		{
 			__m128i Chunk = _mm_loadu_si128(reinterpret_cast<const __m128i*>(SourceData + Cursor));
 			constexpr int32 Mode = _SIDD_UWORD_OPS | _SIDD_CMP_RANGES | _SIDD_MASKED_NEGATIVE_POLARITY;
@@ -258,29 +258,31 @@ static FStringView SkipUntilNonIdentifierCharacter(FStringView Source)
 			if (CompareResult)
 			{
 				Cursor += _mm_cmpistri(NeedleVec, Chunk, Mode);
-				return FStringView(SourceData + Cursor, SourceLen - Cursor);
+				return FShaderSource::FViewType(SourceData + Cursor, SourceLen - Cursor);
 			}
 			Cursor += 8;
 		}
 	}
+	else
 #endif // UE_SHADER_MINIFIER_SSE
-
-	while (Cursor < SourceLen)
 	{
-		if (!IsPossibleIdentifierCharacter(SourceData[Cursor]))
+		while (Cursor < SourceLen)
 		{
-			break;
+			if (!IsPossibleIdentifierCharacter(SourceData[Cursor]))
+			{
+				break;
+			}
+			++Cursor;
 		}
-		++Cursor;
 	}
-	return FStringView(SourceData + Cursor, SourceLen - Cursor);
+	return FShaderSource::FViewType(SourceData + Cursor, FMath::Max(SourceLen - Cursor, 0));
 }
 
-static FStringView SkipUntilNonNumber(FStringView Source) 
+static FShaderSource::FViewType SkipUntilNonNumber(FShaderSource::FViewType Source) 
 {
 	const int32 SourceLen = Source.Len();
 	int32 Cursor = 0;
-	const TCHAR* SourceData = Source.GetData();
+	const FShaderSource::CharType* SourceData = Source.GetData();
 	while (Cursor < SourceLen)
 	{
 		if (!IsNumber(SourceData[Cursor]))
@@ -289,20 +291,20 @@ static FStringView SkipUntilNonNumber(FStringView Source)
 		}
 		++Cursor;
 	}
-	return FStringView(SourceData + Cursor, SourceLen - Cursor);
+	return FShaderSource::FViewType(SourceData + Cursor, SourceLen - Cursor);
 }
 
-static FStringView SkipSpace(FStringView Source)
+static FShaderSource::FViewType SkipSpace(FShaderSource::FViewType Source)
 {
 	const int32 SourceLen = Source.Len();
 	int32 Cursor = 0;
-	const TCHAR* SourceData = Source.GetData();
+	const FShaderSource::CharType* SourceData = Source.GetData();
 
 #if UE_SHADER_MINIFIER_SSE
+	if constexpr (sizeof(FShaderSource::CharType) == 2)
 	{
-		const int32 AlignedLen = SourceLen & (~7); // align down to multiple of 8 TCHAR-s
 		const __m128i NeedleVec = _mm_setr_epi16(L' ', L'\f', L'\r', L'\n', L'\t', L'\v', 0, 0);
-		while (Cursor < AlignedLen)
+		while (Cursor < SourceLen)
 		{
 			__m128i Chunk = _mm_loadu_si128(reinterpret_cast<const __m128i*>(SourceData + Cursor));
 			constexpr int32 Mode = _SIDD_UWORD_OPS | _SIDD_CMP_EQUAL_ANY | _SIDD_MASKED_NEGATIVE_POLARITY;
@@ -310,25 +312,27 @@ static FStringView SkipSpace(FStringView Source)
 			if (CompareResult)
 			{
 				Cursor += _mm_cmpistri(NeedleVec, Chunk, Mode);
-				return FStringView(SourceData + Cursor, SourceLen - Cursor);
+				return FShaderSource::FViewType(SourceData + Cursor, SourceLen - Cursor);
 			}
 			Cursor += 8;
 		}
 	}
+	else
 #endif // UE_SHADER_MINIFIER_SSE
-
-	while (Cursor < SourceLen)
 	{
-		if (!IsSpace(SourceData[Cursor]))
+		while (Cursor < SourceLen)
 		{
-			break;
+			if (!IsSpace(SourceData[Cursor]))
+			{
+				break;
+			}
+			++Cursor;
 		}
-		++Cursor;
 	}
-	return FStringView(SourceData + Cursor, SourceLen - Cursor);
+	return FShaderSource::FViewType(SourceData + Cursor, FMath::Max(SourceLen - Cursor, 0));
 }
 
-static FStringView TrimSpace(FStringView Source)
+static FShaderSource::FViewType TrimSpace(FShaderSource::FViewType Source)
 {
 	int32 CursorBegin = 0;
 	int32 CursorEnd = Source.Len();
@@ -351,30 +355,30 @@ static FStringView TrimSpace(FStringView Source)
 		--CursorEnd;
 	}
 
-	FStringView Result = SubStrView(Source, CursorBegin, CursorEnd-CursorBegin);
+	FShaderSource::FViewType Result = SubStrView(Source, CursorBegin, CursorEnd-CursorBegin);
 
 	return Result;
 }
 
-static FStringView SkipUntilNextLine(FStringView Source)
+static FShaderSource::FViewType SkipUntilNextLine(FShaderSource::FViewType Source)
 {
 	int32 Index = INDEX_NONE;
 	if (Source.FindChar('\n', Index))
 	{
-		return FStringView(Source.GetData() + Index, Source.Len() - Index);
+		return FShaderSource::FViewType(Source.GetData() + Index, Source.Len() - Index);
 	}
 	else
 	{
-		return FStringView {};
+		return FShaderSource::FViewType {};
 	}
 }
 
-static FStringView SkipUntilStr(FStringView Haystack, FStringView Needle)
+static FShaderSource::FViewType SkipUntilStr(FShaderSource::FViewType Haystack, FShaderSource::FViewType Needle)
 {
-	return SkipUntil(Haystack, [Needle](FStringView  S) { return StartsWith(S, Needle); });
+	return SkipUntil(Haystack, [Needle](FShaderSource::FViewType  S) { return StartsWith(S, Needle); });
 }
 
-static FStringView ExtractBlock(FStringView Source, TCHAR DelimBegin, TCHAR DelimEnd)
+static FShaderSource::FViewType ExtractBlock(FShaderSource::FViewType Source, FShaderSource::CharType DelimBegin, FShaderSource::CharType DelimEnd)
 {
 	// TODO: handle comments
 	// TODO: handle #if 0 blocks
@@ -382,7 +386,7 @@ static FStringView ExtractBlock(FStringView Source, TCHAR DelimBegin, TCHAR Deli
 	int32 PosEnd = INDEX_NONE;
 	int32 Stack  = 0;
 	const int32 SourceLen = Source.Len();
-	const TCHAR* SourceData = Source.GetData();
+	const FShaderSource::CharType* SourceData = Source.GetData();
 
 	int32 Cursor = 0;
 
@@ -396,7 +400,7 @@ static FStringView ExtractBlock(FStringView Source, TCHAR DelimBegin, TCHAR Deli
 
 	auto ProcessCharacter = [&Stack, &PosEnd, SourceData, DelimBegin, DelimEnd](int32 Cursor) -> EStatus
 	{
-		TCHAR C = SourceData[Cursor];
+		FShaderSource::CharType C = SourceData[Cursor];
 
 		if (C == DelimBegin)
 		{
@@ -423,10 +427,10 @@ static FStringView ExtractBlock(FStringView Source, TCHAR DelimBegin, TCHAR Deli
 	};
 
 #if UE_SHADER_MINIFIER_SSE
+	if constexpr (sizeof(FShaderSource::CharType) == 2)
 	{
-		const int32 AlignedLen = SourceLen & (~7); // align down to multiple of 8 TCHAR-s
 		const __m128i NeedleVec = _mm_setr_epi16(DelimBegin, DelimEnd, 0, 0, 0, 0, 0, 0);
-		while (Cursor < AlignedLen && Status != EStatus::Finished)
+		while (Cursor < SourceLen && Status != EStatus::Finished)
 		{
 			__m128i Chunk = _mm_loadu_si128(reinterpret_cast<const __m128i*>(SourceData + Cursor));
 			constexpr int32 Mode = _SIDD_UWORD_OPS | _SIDD_CMP_EQUAL_ANY | _SIDD_MOST_SIGNIFICANT;
@@ -446,21 +450,23 @@ static FStringView ExtractBlock(FStringView Source, TCHAR DelimBegin, TCHAR Deli
 			Cursor += 8;
 		}
 	}
+	else
 #endif // UE_SHADER_MINIFIER_SSE
-
-	while (Cursor < SourceLen && Status != EStatus::Finished)
 	{
-		Status = ProcessCharacter(Cursor);
-		++Cursor;
+		while (Cursor < SourceLen && Status != EStatus::Finished)
+		{
+			Status = ProcessCharacter(Cursor);
+			++Cursor;
+		}
 	}
 
 	if (Stack == 0 && PosEnd != INDEX_NONE)
 	{
-		return FStringView(Source.GetData(), PosEnd + 1);
+		return FShaderSource::FViewType(Source.GetData(), PosEnd + 1);
 	}
 	else
 	{
-		return FStringView{};
+		return FShaderSource::FViewType{};
 	}
 }
 
@@ -485,30 +491,30 @@ enum class EBlockType : uint8 {
 
 struct FCodeBlock
 {
-	const TCHAR* CodePtr = nullptr;
+	const FShaderSource::CharType* CodePtr = nullptr;
 	int32        CodeLen = 0;
 	EBlockType   Type = EBlockType::Unknown;
 	uint8        Padding[3] = {};
 
-	operator FStringView () const
+	operator FShaderSource::FViewType () const
 	{
 		return GetCode();
 	}
 
-	bool operator == (const FStringView S) const
+	bool operator == (const FShaderSource::FViewType S) const
 	{
 		return Equals(GetCode(), S);
 	}
 
-	void SetCode(FStringView Code)
+	void SetCode(FShaderSource::FViewType Code)
 	{
 		CodePtr = Code.GetData();
 		CodeLen = Code.Len();
 	}
 
-	FStringView GetCode() const
+	FShaderSource::FViewType GetCode() const
 	{
-		return FStringView(CodePtr, CodeLen);
+		return FShaderSource::FViewType(CodePtr, CodeLen);
 	}
 };
 
@@ -533,14 +539,14 @@ enum class ECodeChunkType {
 struct FNamespace
 {
 	FNamespace() = default;
-	FNamespace(TConstArrayView<FStringView> InStack)
+	FNamespace(TConstArrayView<FShaderSource::FViewType> InStack)
 	{
 		if (!InStack.IsEmpty())
 		{
-			for (const FStringView& Part : InStack)
+			for (const FShaderSource::FViewType& Part : InStack)
 			{
 				FullName += Part;
-				FullName += TEXT("::");
+				FullName += SHADER_SOURCE_LITERAL("::");
 			}
 			FullName.LeftChopInline(2);
 		}
@@ -548,7 +554,7 @@ struct FNamespace
 	}
 
 	FString FullName; // i.e. Foo::Bar::Baz
-	TArray<FStringView> Stack; // i.e. [Foo, Bar, Baz]
+	TArray<FShaderSource::FViewType> Stack; // i.e. [Foo, Bar, Baz]
 };
 
 using FCodeBlockArray = TArray<FCodeBlock, TInlineAllocator<6>>;
@@ -565,7 +571,7 @@ struct FCodeChunk
 	// The struct type may be referenced, but the variable may be removed. In this case we have to emit the type declaration only.
 	bool bVerbatim = true;
 
-	FStringView FindFirstBlockByType(EBlockType InType) const
+	FShaderSource::FViewType FindFirstBlockByType(EBlockType InType) const
 	{
 		for (const FCodeBlock& Block : Blocks)
 		{
@@ -578,7 +584,7 @@ struct FCodeChunk
 	}
 
 	// String view covering the entire code chunk
-	FStringView GetCode() const
+	FShaderSource::FViewType GetCode() const
 	{
 		if (Blocks.IsEmpty())
 		{
@@ -588,31 +594,31 @@ struct FCodeChunk
 		{
 			const FCodeBlock& FirstBlock = Blocks[0];
 			const FCodeBlock& LastBlock = Blocks[Blocks.Num()-1];
-			const TCHAR* Begin = FirstBlock.CodePtr;
-			const TCHAR* End = LastBlock.CodePtr + LastBlock.CodeLen;
-			return FStringView(Begin, int32(End-Begin));
+			const FShaderSource::CharType* Begin = FirstBlock.CodePtr;
+			const FShaderSource::CharType* End = LastBlock.CodePtr + LastBlock.CodeLen;
+			return FShaderSource::FViewType(Begin, int32(End-Begin));
 		}
 	}
 };
 
 struct FParsedShader
 {
-	FStringView Source;
+	FShaderSource::FViewType Source;
 	TArray<FCodeChunk> Chunks;
 	TArray<FNamespace> Namespaces;
-	TArray<FStringView> LineDirectives;
+	TArray<FShaderSource::FViewType> LineDirectives;
 };
 
 struct FNamespaceTracker
 {
 	TMap<FString, int32> UniqueNamespaceMap;
 	TArray<FNamespace> UniqueNamespaceArray;
-	TArray<FStringView> NamespaceStack;
+	TArray<FShaderSource::FViewType> NamespaceStack;
 	TArray<int32> NamespaceIdStack;
 
 	FNamespaceTracker() = default;
 
-	void Push(FStringView Name)
+	void Push(FShaderSource::FViewType Name)
 	{
 		NamespaceStack.Push(Name);
 		FNamespace NamespaceEntry(NamespaceStack);
@@ -645,22 +651,22 @@ struct FNamespaceTracker
 	}
 };
 
-FStringView ExtractNextIdentifier(FStringView Source)
+FShaderSource::FViewType ExtractNextIdentifier(FShaderSource::FViewType Source)
 {
-	FStringView Remainder = SkipUntilNonIdentifierCharacter(Source);
-	FStringView Identifier = SubStrView(Source, 0, Source.Len() - Remainder.Len());
+	FShaderSource::FViewType Remainder = SkipUntilNonIdentifierCharacter(Source);
+	FShaderSource::FViewType Identifier = SubStrView(Source, 0, Source.Len() - Remainder.Len());
 	return Identifier;
 }
 
-static FParsedShader ParseShader(FStringView InSource, FDiagnostics& Output)
+static FParsedShader ParseShader(FShaderSource::FViewType InSource, FDiagnostics& Output)
 {
 	FParsedShader Result;
 
 	Result.Source = InSource;
 
-	FStringView		   Source = InSource;
-	FCodeBlockArray    PendingBlocks;
-	TArray<FCodeChunk> Chunks;
+	FShaderSource::FViewType	Source = InSource;
+	FCodeBlockArray			PendingBlocks;
+	TArray<FCodeChunk>		Chunks;
 
 	ECodeChunkType ChunkType = ECodeChunkType::Unknown;
 	bool		   bFoundBody = false;
@@ -677,7 +683,7 @@ static FParsedShader ParseShader(FStringView InSource, FDiagnostics& Output)
 	int32 OperatorKeywordBlockIndex = INDEX_NONE;
 
 	FNamespaceTracker NamespaceTracker;
-	FStringView		  PendingNamespace;
+	FShaderSource::FViewType PendingNamespace;
 
 	auto AddDiagnostic = [InSource, &Source](TArray<FDiagnosticMessage>& Output, FStringView Message)
 	{
@@ -691,7 +697,7 @@ static FParsedShader ParseShader(FStringView InSource, FDiagnostics& Output)
 		Output.Add(MoveTemp(Diagnostic));
 	};
 
-	auto AddBlock = [&PendingBlocks](EBlockType Type, FStringView Code)
+	auto AddBlock = [&PendingBlocks](EBlockType Type, FShaderSource::FViewType Code)
 	{
 		FCodeBlock NewBlock;
 		NewBlock.Type = Type;
@@ -748,7 +754,7 @@ static FParsedShader ParseShader(FStringView InSource, FDiagnostics& Output)
 				// i.e. `struct Foo { ... } Blah = { expression };` 
 				if (ExpressionBlockIndex > 0 && NameBlockIndex == INDEX_NONE)
 				{
-					AddDiagnostic(Output.Errors, TEXT("Initialized struct variables must be named"));
+					AddDiagnostic(Output.Errors, TEXTVIEW("Initialized struct variables must be named"));
 					return;
 				}
 			}
@@ -882,20 +888,20 @@ static FParsedShader ParseShader(FStringView InSource, FDiagnostics& Output)
 			break;
 		}
 
-		const TCHAR FirstChar = *Source.GetData();
+		const FShaderSource::CharType FirstChar = *Source.GetData();
 
 		if (GCharacterFlags.IsSpecial(FirstChar))
 		{
 			if (FirstChar == '/')
 			{
-				if (StartsWith(Source, TEXTVIEW("//")))
+				if (StartsWith(Source, SHADER_SOURCE_VIEWLITERAL("//")))
 				{
-					FStringView Remainder = SkipUntilNextLine(Source);
+					FShaderSource::FViewType Remainder = SkipUntilNextLine(Source);
 
 					// Save comment lines that are outside of blocks
 					if (PendingBlocks.IsEmpty())
 					{
-						FStringView Block = SubStrView(Source, 0, Source.Len() - Remainder.Len());
+						FShaderSource::FViewType Block = SubStrView(Source, 0, Source.Len() - Remainder.Len());
 						AddBlock(EBlockType::Unknown, Block);
 						ChunkType = ECodeChunkType::CommentLine;
 						FinalizeChunk();
@@ -905,9 +911,9 @@ static FParsedShader ParseShader(FStringView InSource, FDiagnostics& Output)
 
 					continue;
 				}
-				else if (StartsWith(Source, TEXTVIEW("/*")))
+				else if (StartsWith(Source, SHADER_SOURCE_VIEWLITERAL("/*")))
 				{
-					Source = SkipUntilStr(Source, TEXTVIEW("*/"));
+					Source = SkipUntilStr(Source, SHADER_SOURCE_VIEWLITERAL("*/"));
 					if (Source.Len() >= 2)
 					{
 						Source = SubStrView(Source, 2);
@@ -917,38 +923,38 @@ static FParsedShader ParseShader(FStringView InSource, FDiagnostics& Output)
 			}
 			else if (FirstChar == '#')
 			{
-				if (StartsWith(Source, TEXTVIEW("#line")))
+				if (StartsWith(Source, SHADER_SOURCE_VIEWLITERAL("#line")))
 				{
-					FStringView Remainder = SkipUntilNextLine(Source);
-					FStringView Block = SubStrView(Source, 0, Source.Len() - Remainder.Len());
+					FShaderSource::FViewType Remainder = SkipUntilNextLine(Source);
+					FShaderSource::FViewType Block = SubStrView(Source, 0, Source.Len() - Remainder.Len());
 					Result.LineDirectives.Add(Block);
 					Source = Remainder;
 					continue;
 				}
-				else if (StartsWith(Source, TEXTVIEW("#pragma")))
+				else if (StartsWith(Source, SHADER_SOURCE_VIEWLITERAL("#pragma")))
 				{
-					FStringView Remainder = SkipUntilNextLine(Source);
-					FStringView Block = SubStrView(Source, 0, Source.Len() - Remainder.Len());
+					FShaderSource::FViewType Remainder = SkipUntilNextLine(Source);
+					FShaderSource::FViewType Block = SubStrView(Source, 0, Source.Len() - Remainder.Len());
 					AddBlock(EBlockType::Directive, Block);
 					ChunkType = ECodeChunkType::Pragma;
 					FinalizeChunk();
 					Source = Remainder;
 					continue;
 				}
-				else if (StartsWith(Source, TEXTVIEW("#define")))
+				else if (StartsWith(Source, SHADER_SOURCE_VIEWLITERAL("#define")))
 				{
 					// TODO: handle `\` new lines in defines
-					FStringView Remainder = SkipUntilNextLine(Source);
-					FStringView Block = SubStrView(Source, 0, Source.Len() - Remainder.Len());
+					FShaderSource::FViewType Remainder = SkipUntilNextLine(Source);
+					FShaderSource::FViewType Block = SubStrView(Source, 0, Source.Len() - Remainder.Len());
 					AddBlock(EBlockType::Directive, Block);
 					ChunkType = ECodeChunkType::Define;
 					FinalizeChunk();
 					Source = Remainder;
 					continue;
 				}
-				else if (StartsWith(Source, TEXTVIEW("#if 0")))
+				else if (StartsWith(Source, SHADER_SOURCE_VIEWLITERAL("#if 0")))
 				{
-					Source = SkipUntilStr(Source, TEXTVIEW("#endif"));
+					Source = SkipUntilStr(Source, SHADER_SOURCE_VIEWLITERAL("#endif"));
 					if (Source.Len() >= 6)
 					{
 						Source = SubStrView(Source, 6);
@@ -958,13 +964,13 @@ static FParsedShader ParseShader(FStringView InSource, FDiagnostics& Output)
 			}
 			else if (PendingBlocks.IsEmpty() && (FirstChar == '{' || FirstChar == '}'))
 			{
-				if (StartsWith(Source, TEXTVIEW("{")))
+				if (StartsWith(Source, SHADER_SOURCE_VIEWLITERAL("{")))
 				{
 					if (ChunkType == ECodeChunkType::Namespace)
 					{
 						if (PendingNamespace.IsEmpty())
 						{
-							AddDiagnostic(Output.Errors, TEXT("HLSL does not support anonymous namespaces"));
+							AddDiagnostic(Output.Errors, TEXTVIEW("HLSL does not support anonymous namespaces"));
 							break;
 						}
 						else
@@ -978,11 +984,11 @@ static FParsedShader ParseShader(FStringView InSource, FDiagnostics& Output)
 					}
 					else
 					{
-						AddDiagnostic(Output.Errors, TEXT("Expected token '{'"));
+						AddDiagnostic(Output.Errors, TEXTVIEW("Expected token '{'"));
 					}
 					continue;
 				}
-				else if (StartsWith(Source, TEXTVIEW("}")))
+				else if (StartsWith(Source, SHADER_SOURCE_VIEWLITERAL("}")))
 				{
 					if (NamespaceTracker.Pop())
 					{
@@ -991,7 +997,7 @@ static FParsedShader ParseShader(FStringView InSource, FDiagnostics& Output)
 					}
 					else
 					{
-						AddDiagnostic(Output.Errors, TEXT("Expected token '}'"));
+						AddDiagnostic(Output.Errors, TEXTVIEW("Expected token '}'"));
 						break;
 					}
 				}
@@ -1004,11 +1010,11 @@ static FParsedShader ParseShader(FStringView InSource, FDiagnostics& Output)
 		{
 			// Operator keyword found in the last processed block. Expect to find operator name character sequence next.
 
-			FStringView OperatorName = ExtractOperator(Source);
+			FShaderSource::FViewType OperatorName = ExtractOperator(Source);
 
 			if (OperatorName.IsEmpty())
 			{
-				AddDiagnostic(Output.Errors, TEXT("Unexpected operator overload type"));
+				AddDiagnostic(Output.Errors, TEXTVIEW("Unexpected operator overload type"));
 				break;
 			}
 
@@ -1018,55 +1024,55 @@ static FParsedShader ParseShader(FStringView InSource, FDiagnostics& Output)
 			continue;
 		}
 
-		FStringView Remainder = SkipUntilNonIdentifierCharacter(Source);
-		FStringView Identifier = SubStrView(Source, 0, Source.Len() - Remainder.Len());
+		FShaderSource::FViewType Remainder = SkipUntilNonIdentifierCharacter(Source);
+		FShaderSource::FViewType Identifier = SubStrView(Source, 0, Source.Len() - Remainder.Len());
 
 		if (Identifier.Len())
 		{
 			if (ChunkType == ECodeChunkType::Unknown)
 			{
-				if (Equals(Identifier, TEXTVIEW("struct")))
+				if (Equals(Identifier, SHADER_SOURCE_VIEWLITERAL("struct")))
 				{
 					ChunkType = ECodeChunkType::Struct;
 					StructBlockIndex = PendingBlocks.Num();
 				}
-				else if (Equals(Identifier, TEXTVIEW("cbuffer")) || Equals(Identifier, TEXTVIEW("ConstantBuffer")))
+				else if (Equals(Identifier, SHADER_SOURCE_VIEWLITERAL("cbuffer")) || Equals(Identifier, SHADER_SOURCE_VIEWLITERAL("ConstantBuffer")))
 				{
 					ChunkType = ECodeChunkType::CBuffer;
 					CbufferBlockIndex = PendingBlocks.Num();
 				}
-				else if (Equals(Identifier, TEXTVIEW("enum")))
+				else if (Equals(Identifier, SHADER_SOURCE_VIEWLITERAL("enum")))
 				{
 					ChunkType = ECodeChunkType::Enum;
 					EnumBlockIndex = PendingBlocks.Num();
 				}
-				else if (Equals(Identifier, TEXTVIEW("namespace")))
+				else if (Equals(Identifier, SHADER_SOURCE_VIEWLITERAL("namespace")))
 				{
 					ChunkType = ECodeChunkType::Namespace;
 					Source = Remainder;
 					continue;
 				}
-				else if (Equals(Identifier, TEXTVIEW("using")))
+				else if (Equals(Identifier, SHADER_SOURCE_VIEWLITERAL("using")))
 				{
 					ChunkType = ECodeChunkType::Using;
 					Source = Remainder;
 					AddBlock(EBlockType::Keyword, Identifier);
 					continue;
 				}
-				else if (Equals(Identifier, TEXTVIEW("typedef")))
+				else if (Equals(Identifier, SHADER_SOURCE_VIEWLITERAL("typedef")))
 				{
 					ChunkType = ECodeChunkType::Typedef;
 					Source = Remainder;
 					AddBlock(EBlockType::Keyword, Identifier);
 					continue;
 				}
-				else if (Equals(Identifier, TEXTVIEW("template")))
+				else if (Equals(Identifier, SHADER_SOURCE_VIEWLITERAL("template")))
 				{
 					Source = Remainder;
 					AddBlock(EBlockType::Keyword, Identifier);
 					continue;
 				}
-				else if (Equals(Identifier, TEXTVIEW("operator")))
+				else if (Equals(Identifier, SHADER_SOURCE_VIEWLITERAL("operator")))
 				{
 					ChunkType = ECodeChunkType::Operator;
 					Source = Remainder;
@@ -1106,18 +1112,18 @@ static FParsedShader ParseShader(FStringView InSource, FDiagnostics& Output)
 			continue;
 		}
 
-		FStringView Block;
+		FShaderSource::FViewType Block;
 
-		TCHAR C = Source[0];
+		FShaderSource::CharType C = Source[0];
 
 		EBlockType BlockType = EBlockType::Unknown;
 
-		if (StartsWith(Source, TEXTVIEW("==")))
+		if (StartsWith(Source, SHADER_SOURCE_VIEWLITERAL("==")))
 		{
-			AddDiagnostic(Output.Errors, TEXT("Unexpected sequence '=='"));
+			AddDiagnostic(Output.Errors, TEXTVIEW("Unexpected sequence '=='"));
 			break;
 		}
-		else if (StartsWith(Source, TEXTVIEW("::")))
+		else if (StartsWith(Source, SHADER_SOURCE_VIEWLITERAL("::")))
 		{
 			Block = SubStrView(Source, 0, 2);
 			Source = SubStrView(Source, 2);
@@ -1143,9 +1149,9 @@ static FParsedShader ParseShader(FStringView InSource, FDiagnostics& Output)
 			{
 				int32 Pos = INDEX_NONE;
 
-				if (!Source.FindChar(TCHAR(';'), Pos))
+				if (!Source.FindChar(FShaderSource::CharType(';'), Pos))
 				{
-					AddDiagnostic(Output.Errors, TEXT("Expected semicolon after assignment expression"));
+					AddDiagnostic(Output.Errors, TEXTVIEW("Expected semicolon after assignment expression"));
 					break;
 				}
 
@@ -1238,7 +1244,7 @@ static FParsedShader ParseShader(FStringView InSource, FDiagnostics& Output)
 
 		if (Block.IsEmpty())
 		{
-			AddDiagnostic(Output.Errors, TEXT("Failed to extract code block"));
+			AddDiagnostic(Output.Errors, TEXTVIEW("Failed to extract code block"));
 			break;
 		}
 		else
@@ -1269,7 +1275,7 @@ static FParsedShader ParseShader(FStringView InSource, FDiagnostics& Output)
 }
 
 template<typename CallbackT>
-void FindChunksByIdentifier(TConstArrayView<FCodeChunk> Chunks, FStringView Identifier, CallbackT Callback)
+void FindChunksByIdentifier(TConstArrayView<FCodeChunk> Chunks, FShaderSource::FViewType Identifier, CallbackT Callback)
 {
 	for (const FCodeChunk& Chunk : Chunks)
 	{
@@ -1283,16 +1289,16 @@ void FindChunksByIdentifier(TConstArrayView<FCodeChunk> Chunks, FStringView Iden
 	}
 }
 
-static TArray<FStringView> SplitByChar(FStringView Source, TCHAR Delimiter)
+static TArray<FShaderSource::FViewType> SplitByChar(FShaderSource::FViewType Source, FShaderSource::CharType Delimiter)
 {
-	TArray<FStringView> Result;
+	TArray<FShaderSource::FViewType> Result;
 
 	int32 Start = 0;
 	const int32 SourceLen = Source.Len();
 
 	for (int32 I = 0; I < SourceLen; ++I)
 	{
-		TCHAR C = Source[I];
+		FShaderSource::CharType C = Source[I];
 		if (C == Delimiter)
 		{
 			size_t Len = I - Start;
@@ -1310,26 +1316,26 @@ static TArray<FStringView> SplitByChar(FStringView Source, TCHAR Delimiter)
 	return Result;
 }
 
-static void ExtractIdentifiers(FStringView InSource, TArray<FStringView>& Result)
+static void ExtractIdentifiers(FShaderSource::FViewType InSource, TArray<FShaderSource::FViewType>& Result)
 {
-	FStringView Source = InSource;
+	FShaderSource::FViewType Source = InSource;
 
 	while (!Source.IsEmpty())
 	{
-		TCHAR FirstChar = Source.GetData()[0];
+		FShaderSource::CharType FirstChar = Source.GetData()[0];
 
-		if (FirstChar == L'#')
+		if (FirstChar == '#')
 		{
-			if (StartsWith(Source, TEXTVIEW("#line"))
-				|| StartsWith(Source, TEXTVIEW("#pragma")))
+			if (StartsWith(Source, SHADER_SOURCE_VIEWLITERAL("#line"))
+				|| StartsWith(Source, SHADER_SOURCE_VIEWLITERAL("#pragma")))
 			{
 				Source = SkipUntilNextLine(Source);
 				Source = SkipSpace(Source);
 				continue;
 			}
-			else if (StartsWith(Source, TEXTVIEW("#if 0")))
+			else if (StartsWith(Source, SHADER_SOURCE_VIEWLITERAL("#if 0")))
 			{
-				Source = SkipUntilStr(Source, TEXTVIEW("#endif"));
+				Source = SkipUntilStr(Source, SHADER_SOURCE_VIEWLITERAL("#endif"));
 				if (Source.Len() >= 6)
 				{
 					Source = SubStrView(Source, 6);
@@ -1338,17 +1344,17 @@ static void ExtractIdentifiers(FStringView InSource, TArray<FStringView>& Result
 				continue;
 			}
 		}
-		else if (FirstChar == L'/')
+		else if (FirstChar == '/')
 		{
-			if (StartsWith(Source, TEXTVIEW("//")))
+			if (StartsWith(Source, SHADER_SOURCE_VIEWLITERAL("//")))
 			{
 				Source = SkipUntilNextLine(Source);
 				Source = SkipSpace(Source);
 				continue;
 			}
-			else if (StartsWith(Source, TEXTVIEW("/*")))
+			else if (StartsWith(Source, SHADER_SOURCE_VIEWLITERAL("/*")))
 			{
-				Source = SkipUntilStr(Source, TEXTVIEW("*/"));
+				Source = SkipUntilStr(Source, SHADER_SOURCE_VIEWLITERAL("*/"));
 				if (Source.Len() >= 2)
 				{
 					Source = SubStrView(Source, 2);
@@ -1358,9 +1364,9 @@ static void ExtractIdentifiers(FStringView InSource, TArray<FStringView>& Result
 			}
 		}
 
-		FStringView Remainder = SkipUntilNonIdentifierCharacter(Source);
+		FShaderSource::FViewType Remainder = SkipUntilNonIdentifierCharacter(Source);
 
-		FStringView Identifier = SubStrView(Source, 0, Source.Len() - Remainder.Len());
+		FShaderSource::FViewType Identifier = SubStrView(Source, 0, Source.Len() - Remainder.Len());
 
 		if (Identifier.IsEmpty())
 		{
@@ -1381,7 +1387,7 @@ static void ExtractIdentifiers(FStringView InSource, TArray<FStringView>& Result
 	}
 }
 
-static void ExtractIdentifiers(const FCodeChunk& Chunk, TArray<FStringView>& Result)
+static void ExtractIdentifiers(const FCodeChunk& Chunk, TArray<FShaderSource::FViewType>& Result)
 {
 	for (const FCodeBlock& Block : Chunk.Blocks)
 	{
@@ -1389,7 +1395,7 @@ static void ExtractIdentifiers(const FCodeChunk& Chunk, TArray<FStringView>& Res
 	}
 }
 
-static void OutputChunk(const FCodeChunk& Chunk, FString& OutputStream)
+static void OutputChunk(const FCodeChunk& Chunk, FShaderSource::FStringType& OutputStream)
 {
 	if (Chunk.Blocks.IsEmpty())
 	{
@@ -1408,21 +1414,21 @@ static void OutputChunk(const FCodeChunk& Chunk, FString& OutputStream)
 		{
 			if (Index != 0)
 			{
-				OutputStream.AppendChar(L' ');
+				OutputStream.AppendChar(' ');
 			}
 
 			if (Block.Type == EBlockType::Expression)
 			{
-				OutputStream.Append(TEXTVIEW("= "));
+				OutputStream.Append(SHADER_SOURCE_VIEWLITERAL("= "));
 			}
 			else if (Block.Type == EBlockType::Body)
 			{
-				OutputStream.AppendChar(L'\n');
+				OutputStream.AppendChar('\n');
 			}
 
 			if (Block.Type == EBlockType::Binding || Block.Type == EBlockType::Base)
 			{
-				OutputStream.Append(TEXTVIEW(": "));
+				OutputStream.Append(SHADER_SOURCE_VIEWLITERAL(": "));
 			}
 
 			OutputStream.Append(Block.GetCode());
@@ -1438,69 +1444,72 @@ static void OutputChunk(const FCodeChunk& Chunk, FString& OutputStream)
 		&& Chunk.Type != ECodeChunkType::Define
 		&& Chunk.Type != ECodeChunkType::CommentLine)
 	{
-		OutputStream.AppendChar(L';');
+		OutputStream.AppendChar(';');
 	}
 
-	OutputStream.AppendChar(L'\n');
+	OutputStream.AppendChar('\n');
 }
 
-struct FCasedStringViewKeyFuncs : public DefaultKeyFuncs<FStringView>
+struct FCasedStringViewKeyFuncs : public DefaultKeyFuncs<FShaderSource::FViewType>
 {
-	static FORCEINLINE FStringView GetSetKey(FStringView K) { return K; }
+	static FORCEINLINE FShaderSource::FViewType GetSetKey(FShaderSource::FViewType K) { return K; }
 	template <typename T>
-	static FORCEINLINE FStringView GetSetKey(const TPair<FStringView, T>& P) { return P.Key; }
-	static FORCEINLINE bool Matches(FStringView A, FStringView B) { return Equals(A, B); }
-	static FORCEINLINE uint32 GetKeyHash(FStringView Key)
+	static FORCEINLINE FShaderSource::FViewType GetSetKey(const TPair<FShaderSource::FViewType, T>& P) { return P.Key; }
+	static FORCEINLINE bool Matches(FShaderSource::FViewType A, FShaderSource::FViewType B) { return Equals(A, B); }
+	static FORCEINLINE uint32 GetKeyHash(FShaderSource::FViewType Key)
 	{
 		return FXxHash64::HashBuffer(Key.GetData(), Key.Len() * sizeof(*Key.GetData())).Hash;
 	}
 };
 
-static void BuildLineBreakMap(FStringView Source, TArray<int32>& OutLineBreakMap)
+static void BuildLineBreakMap(FShaderSource::FViewType Source, TArray<int32>& OutLineBreakMap)
 {
 	OutLineBreakMap.Reset();
 
 	OutLineBreakMap.Add(0); // Lines numbers are 1-based, so add a dummy element to make UpperBound later return the line number directly
 
 	const int32 SourceLen = Source.Len();
-	const TCHAR* Chars = Source.GetData(); // avoid bounds check overhead in [] operator
+	const FShaderSource::CharType* Chars = Source.GetData(); // avoid bounds check overhead in [] operator
 
 	int32 Cursor = 0;
 
 #if UE_SHADER_MINIFIER_SSE
-	static_assert(sizeof(*Chars) == 2, "BuildLineBreakMap expects 16 bit characters");
-	const int32 AlignedLen = SourceLen & (~7); // align down to multiple of 8 TCHAR-s
-	const __m128i Needle = _mm_set1_epi16(L'\n');
-	while (Cursor < AlignedLen)
+	if constexpr (sizeof(FShaderSource::CharType) == 2)
 	{
-		__m128i Chunk = _mm_loadu_si128(reinterpret_cast<const __m128i*>(Chars + Cursor));
-		__m128i MaskVec = _mm_cmpeq_epi16(Chunk, Needle);
-		uint32 Mask = _mm_movemask_epi8(MaskVec);
-		while (Mask != 0)
+		const __m128i Needle = _mm_set1_epi16(L'\n');
+		while (Cursor < SourceLen)
 		{
-			// NOTE: 2 bits represent each character
-			const uint32 BitIndex = FMath::CountTrailingZeros(Mask);
-			const uint32 ChunkCharIndex = BitIndex / 2;
-			OutLineBreakMap.Add(Cursor + ChunkCharIndex);
-			Mask &= ~(3 << BitIndex);
+			__m128i Chunk = _mm_loadu_si128(reinterpret_cast<const __m128i*>(Chars + Cursor));
+			__m128i MaskVec = _mm_cmpeq_epi16(Chunk, Needle);
+			uint32 Mask = _mm_movemask_epi8(MaskVec);
+			while (Mask != 0)
+			{
+				// NOTE: 2 bits represent each character
+				const uint32 BitIndex = FMath::CountTrailingZeros(Mask);
+				const uint32 ChunkCharIndex = BitIndex / 2;
+				OutLineBreakMap.Add(Cursor + ChunkCharIndex);
+				Mask &= ~(3 << BitIndex);
+			}
+			Cursor += 8;
 		}
-		Cursor += 8;
 	}
+	else
 #endif //UE_SHADER_MINIFIER_SSE
-
-	while (Cursor < SourceLen)
 	{
-		if (Chars[Cursor] == TCHAR('\n'))
+		while (Cursor < SourceLen)
 		{
-			OutLineBreakMap.Add(Cursor);
+			if (Chars[Cursor] == FShaderSource::CharType('\n'))
+			{
+				OutLineBreakMap.Add(Cursor);
+			}
+			++Cursor;
 		}
-		++Cursor;
 	}
 }
 
-static int32 FindLineDirective(const TArray<FStringView>& LineDirectives, const TCHAR* Ptr)
+static int32 FindLineDirective(const TArray<FShaderSource::FViewType>& LineDirectives, const FShaderSource::CharType* Ptr)
 {
-	int32 FoundIndex = Algo::UpperBoundBy(LineDirectives, Ptr, [](FStringView Item)
+	int32 FoundIndex = Algo::UpperBoundBy(LineDirectives, Ptr, [](FShaderSource::FViewType Item)
 	{
 		return Item.GetData();
 	});
@@ -1514,7 +1523,7 @@ static int32 FindLineDirective(const TArray<FStringView>& LineDirectives, const
 	return FoundIndex - 1;
 }
 
-static int32 FindLineNumber(FStringView Source, const TArray<int32>& LineBreakMap, const TCHAR* Ptr)
+static int32 FindLineNumber(FShaderSource::FViewType Source, const TArray<int32>& LineBreakMap, const FShaderSource::CharType* Ptr)
 {
 	if (Ptr < Source.GetData() || Ptr >= Source.GetData() + Source.Len())
 	{
@@ -1528,9 +1537,9 @@ static int32 FindLineNumber(FStringView Source, const TArray<int32>& LineBreakMa
 	return FoundLineNumber;
 }
 
-static bool ParseLineDirective(FStringView Input, int32& OutLineNumber, FStringView& OutFileName)
+static bool ParseLineDirective(FShaderSource::FViewType Input, int32& OutLineNumber, FShaderSource::FViewType& OutFileName)
 {
-	if (!StartsWith(Input, TEXTVIEW("#line")))
+	if (!StartsWith(Input, SHADER_SOURCE_VIEWLITERAL("#line")))
 	{
 		return false;
 	}
@@ -1543,14 +1552,14 @@ static bool ParseLineDirective(FStringView Input, int32& OutLineNumber, FStringV
 		return false;
 	}
 
-	OutLineNumber = FCString::Atoi(Input.GetData());
+	OutLineNumber = FShaderSource::FCStringType::Atoi(Input.GetData());
 
 	int32 FileNameBeginIndex = INDEX_NONE;
-	if (Input.FindChar(TCHAR('"'), FileNameBeginIndex))
+	if (Input.FindChar(FShaderSource::CharType('"'), FileNameBeginIndex))
 	{
 		int32 FileNameEndIndex = INDEX_NONE;
 		Input.MidInline(FileNameBeginIndex + 1);
-		if (Input.FindChar(TCHAR('"'), FileNameEndIndex))
+		if (Input.FindChar(FShaderSource::CharType('"'), FileNameEndIndex))
 		{
 			OutFileName = Input.Mid(0, FileNameEndIndex);
 		}
@@ -1563,41 +1572,41 @@ static bool ParseLineDirective(FStringView Input, int32& OutLineNumber, FStringV
 	return true;
 }
 
-static void OpenNamespace(FString& OutputStream, const FNamespace& Namespace)
+static void OpenNamespace(FShaderSource::FStringType& OutputStream, const FNamespace& Namespace)
 {
-	for (const FStringView& Name : Namespace.Stack)
+	for (const FShaderSource::FViewType& Name : Namespace.Stack)
 	{
-		OutputStream.Append(TEXTVIEW("namespace "));
+		OutputStream.Append(SHADER_SOURCE_VIEWLITERAL("namespace "));
 		OutputStream.Append(Name);
-		OutputStream.Append(TEXTVIEW(" { "));
+		OutputStream.Append(SHADER_SOURCE_VIEWLITERAL(" { "));
 	}
 }
 
-static void CloseNamespace(FString& OutputStream, const FNamespace& Namespace)
+static void CloseNamespace(FShaderSource::FStringType& OutputStream, const FNamespace& Namespace)
 {
-	for (const FStringView& Name : Namespace.Stack)
+	for (const FShaderSource::FViewType& Name : Namespace.Stack)
 	{
-		OutputStream.AppendChar(L'}');
+		OutputStream.AppendChar('}');
 	}
 
-	OutputStream.Append(TEXTVIEW(" // namespace "));
+	OutputStream.Append(SHADER_SOURCE_VIEWLITERAL(" // namespace "));
 	OutputStream.Append(Namespace.FullName);
 }
 
-static FString MinifyShader(const FParsedShader& Parsed, TConstArrayView<FStringView> RequiredSymbols, EMinifyShaderFlags Flags, FDiagnostics& Diagnostics)
+static FShaderSource::FStringType MinifyShader(const FParsedShader& Parsed, TConstArrayView<FShaderSource::FViewType> RequiredSymbols, EMinifyShaderFlags Flags, FDiagnostics& Diagnostics)
 {
-	FString OutputStream;
+	FShaderSource::FStringType OutputStream;
 
 	OutputStream.Reserve(Parsed.Source.Len() / 3); // Heuristic pre-allocation based on average measured reduced code size
 
-	TSet<FStringView, FCasedStringViewKeyFuncs, FDefaultSetAllocator> RelevantIdentifiers;
+	TSet<FShaderSource::FViewType, FCasedStringViewKeyFuncs, FDefaultSetAllocator> RelevantIdentifiers;
 
 	TSet<const FCodeChunk*> RelevantChunks;
-	TSet<FStringView, FCasedStringViewKeyFuncs, FDefaultSetAllocator>  ProcessedIdentifiers;
+	TSet<FShaderSource::FViewType, FCasedStringViewKeyFuncs, FDefaultSetAllocator>  ProcessedIdentifiers;
 
 	TArray<const FCodeChunk*> PendingChunks;
 
-	for (FStringView Entry : RequiredSymbols)
+	for (FShaderSource::FViewType Entry : RequiredSymbols)
 	{
 		RelevantIdentifiers.Add(Entry);
 		ProcessedIdentifiers.Add(Entry);
@@ -1611,141 +1620,141 @@ static FString MinifyShader(const FParsedShader& Parsed, TConstArrayView<FString
 
 	{
 		// Some known builtin words to ignore
-		ProcessedIdentifiers.Add(TEXT("asfloat"));
-		ProcessedIdentifiers.Add(TEXT("asint"));
-		ProcessedIdentifiers.Add(TEXT("asuint"));
-		ProcessedIdentifiers.Add(TEXT("bool"));
-		ProcessedIdentifiers.Add(TEXT("bool2"));
-		ProcessedIdentifiers.Add(TEXT("bool3"));
-		ProcessedIdentifiers.Add(TEXT("bool4"));
-		ProcessedIdentifiers.Add(TEXT("break"));
-		ProcessedIdentifiers.Add(TEXT("cbuffer"));
-		ProcessedIdentifiers.Add(TEXT("const"));
-		ProcessedIdentifiers.Add(TEXT("else"));
-		ProcessedIdentifiers.Add(TEXT("extern"));
-		ProcessedIdentifiers.Add(TEXT("false"));
-		ProcessedIdentifiers.Add(TEXT("float"));
-		ProcessedIdentifiers.Add(TEXT("float2"));
-		ProcessedIdentifiers.Add(TEXT("float3"));
-		ProcessedIdentifiers.Add(TEXT("float3x3"));
-		ProcessedIdentifiers.Add(TEXT("float3x4"));
-		ProcessedIdentifiers.Add(TEXT("float4"));
-		ProcessedIdentifiers.Add(TEXT("float4x4"));
-		ProcessedIdentifiers.Add(TEXT("for"));
-		ProcessedIdentifiers.Add(TEXT("groupshared"));
-		ProcessedIdentifiers.Add(TEXT("if"));
-		ProcessedIdentifiers.Add(TEXT("in"));
-		ProcessedIdentifiers.Add(TEXT("inout"));
-		ProcessedIdentifiers.Add(TEXT("int"));
-		ProcessedIdentifiers.Add(TEXT("int2"));
-		ProcessedIdentifiers.Add(TEXT("int3"));
-		ProcessedIdentifiers.Add(TEXT("int4"));
-		ProcessedIdentifiers.Add(TEXT("interface"));
-		ProcessedIdentifiers.Add(TEXT("out"));
-		ProcessedIdentifiers.Add(TEXT("packoffset"));
-		ProcessedIdentifiers.Add(TEXT("precise"));
-		ProcessedIdentifiers.Add(TEXT("register"));
-		ProcessedIdentifiers.Add(TEXT("return"));
-		ProcessedIdentifiers.Add(TEXT("static"));
-		ProcessedIdentifiers.Add(TEXT("struct"));
-		ProcessedIdentifiers.Add(TEXT("switch"));
-		ProcessedIdentifiers.Add(TEXT("tbuffer"));
-		ProcessedIdentifiers.Add(TEXT("true"));
-		ProcessedIdentifiers.Add(TEXT("uint"));
-		ProcessedIdentifiers.Add(TEXT("uint2"));
-		ProcessedIdentifiers.Add(TEXT("uint3"));
-		ProcessedIdentifiers.Add(TEXT("uint4"));
-		ProcessedIdentifiers.Add(TEXT("void"));
-		ProcessedIdentifiers.Add(TEXT("while"));
-		ProcessedIdentifiers.Add(TEXT("typedef"));
-		ProcessedIdentifiers.Add(TEXT("template"));
-		ProcessedIdentifiers.Add(TEXT("operator"));
-		ProcessedIdentifiers.Add(TEXT("enum"));
+		ProcessedIdentifiers.Add(SHADER_SOURCE_LITERAL("asfloat"));
+		ProcessedIdentifiers.Add(SHADER_SOURCE_LITERAL("asint"));
+		ProcessedIdentifiers.Add(SHADER_SOURCE_LITERAL("asuint"));
+		ProcessedIdentifiers.Add(SHADER_SOURCE_LITERAL("bool"));
+		ProcessedIdentifiers.Add(SHADER_SOURCE_LITERAL("bool2"));
+		ProcessedIdentifiers.Add(SHADER_SOURCE_LITERAL("bool3"));
+		ProcessedIdentifiers.Add(SHADER_SOURCE_LITERAL("bool4"));
+		ProcessedIdentifiers.Add(SHADER_SOURCE_LITERAL("break"));
+		ProcessedIdentifiers.Add(SHADER_SOURCE_LITERAL("cbuffer"));
+		ProcessedIdentifiers.Add(SHADER_SOURCE_LITERAL("const"));
+		ProcessedIdentifiers.Add(SHADER_SOURCE_LITERAL("else"));
+		ProcessedIdentifiers.Add(SHADER_SOURCE_LITERAL("extern"));
+		ProcessedIdentifiers.Add(SHADER_SOURCE_LITERAL("false"));
+		ProcessedIdentifiers.Add(SHADER_SOURCE_LITERAL("float"));
+		ProcessedIdentifiers.Add(SHADER_SOURCE_LITERAL("float2"));
+		ProcessedIdentifiers.Add(SHADER_SOURCE_LITERAL("float3"));
+		ProcessedIdentifiers.Add(SHADER_SOURCE_LITERAL("float3x3"));
+		ProcessedIdentifiers.Add(SHADER_SOURCE_LITERAL("float3x4"));
+		ProcessedIdentifiers.Add(SHADER_SOURCE_LITERAL("float4"));
+		ProcessedIdentifiers.Add(SHADER_SOURCE_LITERAL("float4x4"));
+		ProcessedIdentifiers.Add(SHADER_SOURCE_LITERAL("for"));
+		ProcessedIdentifiers.Add(SHADER_SOURCE_LITERAL("groupshared"));
+		ProcessedIdentifiers.Add(SHADER_SOURCE_LITERAL("if"));
+		ProcessedIdentifiers.Add(SHADER_SOURCE_LITERAL("in"));
+		ProcessedIdentifiers.Add(SHADER_SOURCE_LITERAL("inout"));
+		ProcessedIdentifiers.Add(SHADER_SOURCE_LITERAL("int"));
+		ProcessedIdentifiers.Add(SHADER_SOURCE_LITERAL("int2"));
+		ProcessedIdentifiers.Add(SHADER_SOURCE_LITERAL("int3"));
+		ProcessedIdentifiers.Add(SHADER_SOURCE_LITERAL("int4"));
+		ProcessedIdentifiers.Add(SHADER_SOURCE_LITERAL("interface"));
+		ProcessedIdentifiers.Add(SHADER_SOURCE_LITERAL("out"));
+		ProcessedIdentifiers.Add(SHADER_SOURCE_LITERAL("packoffset"));
+		ProcessedIdentifiers.Add(SHADER_SOURCE_LITERAL("precise"));
+		ProcessedIdentifiers.Add(SHADER_SOURCE_LITERAL("register"));
+		ProcessedIdentifiers.Add(SHADER_SOURCE_LITERAL("return"));
+		ProcessedIdentifiers.Add(SHADER_SOURCE_LITERAL("static"));
+		ProcessedIdentifiers.Add(SHADER_SOURCE_LITERAL("struct"));
+		ProcessedIdentifiers.Add(SHADER_SOURCE_LITERAL("switch"));
+		ProcessedIdentifiers.Add(SHADER_SOURCE_LITERAL("tbuffer"));
+		ProcessedIdentifiers.Add(SHADER_SOURCE_LITERAL("true"));
+		ProcessedIdentifiers.Add(SHADER_SOURCE_LITERAL("uint"));
+		ProcessedIdentifiers.Add(SHADER_SOURCE_LITERAL("uint2"));
+		ProcessedIdentifiers.Add(SHADER_SOURCE_LITERAL("uint3"));
+		ProcessedIdentifiers.Add(SHADER_SOURCE_LITERAL("uint4"));
+		ProcessedIdentifiers.Add(SHADER_SOURCE_LITERAL("void"));
+		ProcessedIdentifiers.Add(SHADER_SOURCE_LITERAL("while"));
+		ProcessedIdentifiers.Add(SHADER_SOURCE_LITERAL("typedef"));
+		ProcessedIdentifiers.Add(SHADER_SOURCE_LITERAL("template"));
+		ProcessedIdentifiers.Add(SHADER_SOURCE_LITERAL("operator"));
+		ProcessedIdentifiers.Add(SHADER_SOURCE_LITERAL("enum"));
 
 		// HLSL resource types
-		ProcessedIdentifiers.Add(TEXT("TextureCubeArray"));
-		ProcessedIdentifiers.Add(TEXT("TextureCube"));
-		ProcessedIdentifiers.Add(TEXT("TextureBuffer"));
-		ProcessedIdentifiers.Add(TEXT("Texture3D"));
-		ProcessedIdentifiers.Add(TEXT("Texture2DMSArray"));
-		ProcessedIdentifiers.Add(TEXT("Texture2DMS"));
-		ProcessedIdentifiers.Add(TEXT("Texture2DArray"));
-		ProcessedIdentifiers.Add(TEXT("Texture2D"));
-		ProcessedIdentifiers.Add(TEXT("Texture1DArray"));
-		ProcessedIdentifiers.Add(TEXT("Texture1D"));
-		ProcessedIdentifiers.Add(TEXT("StructuredBuffer"));
-		ProcessedIdentifiers.Add(TEXT("SamplerState"));
-		ProcessedIdentifiers.Add(TEXT("SamplerComparisonState"));
-		ProcessedIdentifiers.Add(TEXT("RWTextureCubeArray"));
-		ProcessedIdentifiers.Add(TEXT("RWTextureCube"));
-		ProcessedIdentifiers.Add(TEXT("RWTexture3D"));
-		ProcessedIdentifiers.Add(TEXT("RWTexture2DMSArray"));
-		ProcessedIdentifiers.Add(TEXT("RWTexture2DMS"));
-		ProcessedIdentifiers.Add(TEXT("RWTexture2DArray"));
-		ProcessedIdentifiers.Add(TEXT("RWTexture2D"));
-		ProcessedIdentifiers.Add(TEXT("RWTexture1DArray"));
-		ProcessedIdentifiers.Add(TEXT("RWTexture1D"));
-		ProcessedIdentifiers.Add(TEXT("RWStructuredBuffer"));
-		ProcessedIdentifiers.Add(TEXT("RWByteAddressBuffer"));
-		ProcessedIdentifiers.Add(TEXT("RWBuffer"));
-		ProcessedIdentifiers.Add(TEXT("RaytracingAccelerationStructure"));
-		ProcessedIdentifiers.Add(TEXT("RasterizerOrderedTexture3D"));
-		ProcessedIdentifiers.Add(TEXT("RasterizerOrderedTexture2DArray"));
-		ProcessedIdentifiers.Add(TEXT("RasterizerOrderedTexture2D"));
-		ProcessedIdentifiers.Add(TEXT("RasterizerOrderedTexture1DArray"));
-		ProcessedIdentifiers.Add(TEXT("RasterizerOrderedTexture1D"));
-		ProcessedIdentifiers.Add(TEXT("RasterizerOrderedStructuredBuffer"));
-		ProcessedIdentifiers.Add(TEXT("RasterizerOrderedByteAddressBuffer"));
-		ProcessedIdentifiers.Add(TEXT("RasterizerOrderedBuffer"));
-		ProcessedIdentifiers.Add(TEXT("FeedbackTexture2DArray"));
-		ProcessedIdentifiers.Add(TEXT("FeedbackTexture2D"));
-		ProcessedIdentifiers.Add(TEXT("ConsumeStructuredBuffer"));
-		ProcessedIdentifiers.Add(TEXT("ConstantBuffer"));
-		ProcessedIdentifiers.Add(TEXT("ByteAddressBuffer"));
-		ProcessedIdentifiers.Add(TEXT("Buffer"));
-		ProcessedIdentifiers.Add(TEXT("AppendStructuredBuffer"));
+		ProcessedIdentifiers.Add(SHADER_SOURCE_LITERAL("TextureCubeArray"));
+		ProcessedIdentifiers.Add(SHADER_SOURCE_LITERAL("TextureCube"));
+		ProcessedIdentifiers.Add(SHADER_SOURCE_LITERAL("TextureBuffer"));
+		ProcessedIdentifiers.Add(SHADER_SOURCE_LITERAL("Texture3D"));
+		ProcessedIdentifiers.Add(SHADER_SOURCE_LITERAL("Texture2DMSArray"));
+		ProcessedIdentifiers.Add(SHADER_SOURCE_LITERAL("Texture2DMS"));
+		ProcessedIdentifiers.Add(SHADER_SOURCE_LITERAL("Texture2DArray"));
+		ProcessedIdentifiers.Add(SHADER_SOURCE_LITERAL("Texture2D"));
+		ProcessedIdentifiers.Add(SHADER_SOURCE_LITERAL("Texture1DArray"));
+		ProcessedIdentifiers.Add(SHADER_SOURCE_LITERAL("Texture1D"));
+		ProcessedIdentifiers.Add(SHADER_SOURCE_LITERAL("StructuredBuffer"));
+		ProcessedIdentifiers.Add(SHADER_SOURCE_LITERAL("SamplerState"));
+		ProcessedIdentifiers.Add(SHADER_SOURCE_LITERAL("SamplerComparisonState"));
+		ProcessedIdentifiers.Add(SHADER_SOURCE_LITERAL("RWTextureCubeArray"));
+		ProcessedIdentifiers.Add(SHADER_SOURCE_LITERAL("RWTextureCube"));
+		ProcessedIdentifiers.Add(SHADER_SOURCE_LITERAL("RWTexture3D"));
+		ProcessedIdentifiers.Add(SHADER_SOURCE_LITERAL("RWTexture2DMSArray"));
+		ProcessedIdentifiers.Add(SHADER_SOURCE_LITERAL("RWTexture2DMS"));
+		ProcessedIdentifiers.Add(SHADER_SOURCE_LITERAL("RWTexture2DArray"));
+		ProcessedIdentifiers.Add(SHADER_SOURCE_LITERAL("RWTexture2D"));
+		ProcessedIdentifiers.Add(SHADER_SOURCE_LITERAL("RWTexture1DArray"));
+		ProcessedIdentifiers.Add(SHADER_SOURCE_LITERAL("RWTexture1D"));
+		ProcessedIdentifiers.Add(SHADER_SOURCE_LITERAL("RWStructuredBuffer"));
+		ProcessedIdentifiers.Add(SHADER_SOURCE_LITERAL("RWByteAddressBuffer"));
+		ProcessedIdentifiers.Add(SHADER_SOURCE_LITERAL("RWBuffer"));
+		ProcessedIdentifiers.Add(SHADER_SOURCE_LITERAL("RaytracingAccelerationStructure"));
+		ProcessedIdentifiers.Add(SHADER_SOURCE_LITERAL("RasterizerOrderedTexture3D"));
+		ProcessedIdentifiers.Add(SHADER_SOURCE_LITERAL("RasterizerOrderedTexture2DArray"));
+		ProcessedIdentifiers.Add(SHADER_SOURCE_LITERAL("RasterizerOrderedTexture2D"));
+		ProcessedIdentifiers.Add(SHADER_SOURCE_LITERAL("RasterizerOrderedTexture1DArray"));
+		ProcessedIdentifiers.Add(SHADER_SOURCE_LITERAL("RasterizerOrderedTexture1D"));
+		ProcessedIdentifiers.Add(SHADER_SOURCE_LITERAL("RasterizerOrderedStructuredBuffer"));
+		ProcessedIdentifiers.Add(SHADER_SOURCE_LITERAL("RasterizerOrderedByteAddressBuffer"));
+		ProcessedIdentifiers.Add(SHADER_SOURCE_LITERAL("RasterizerOrderedBuffer"));
+		ProcessedIdentifiers.Add(SHADER_SOURCE_LITERAL("FeedbackTexture2DArray"));
+		ProcessedIdentifiers.Add(SHADER_SOURCE_LITERAL("FeedbackTexture2D"));
+		ProcessedIdentifiers.Add(SHADER_SOURCE_LITERAL("ConsumeStructuredBuffer"));
+		ProcessedIdentifiers.Add(SHADER_SOURCE_LITERAL("ConstantBuffer"));
+		ProcessedIdentifiers.Add(SHADER_SOURCE_LITERAL("ByteAddressBuffer"));
+		ProcessedIdentifiers.Add(SHADER_SOURCE_LITERAL("Buffer"));
+		ProcessedIdentifiers.Add(SHADER_SOURCE_LITERAL("AppendStructuredBuffer"));
 
 		// Alternative spelling of some resource types
-		ProcessedIdentifiers.Add(TEXT("AppendRegularBuffer"));
-		ProcessedIdentifiers.Add(TEXT("ByteBuffer"));
-		ProcessedIdentifiers.Add(TEXT("ConsumeRegularBuffer"));
-		ProcessedIdentifiers.Add(TEXT("DataBuffer"));
-		ProcessedIdentifiers.Add(TEXT("MS_Texture2D"));
-		ProcessedIdentifiers.Add(TEXT("MS_Texture2D_Array"));
-		ProcessedIdentifiers.Add(TEXT("RegularBuffer"));
-		ProcessedIdentifiers.Add(TEXT("RW_ByteBuffer"));
-		ProcessedIdentifiers.Add(TEXT("RW_DataBuffer"));
-		ProcessedIdentifiers.Add(TEXT("RW_RegularBuffer"));
-		ProcessedIdentifiers.Add(TEXT("RW_Texture1D"));
-		ProcessedIdentifiers.Add(TEXT("RW_Texture1D_Array"));
-		ProcessedIdentifiers.Add(TEXT("RW_Texture2D"));
-		ProcessedIdentifiers.Add(TEXT("RW_Texture2D_Array"));
-		ProcessedIdentifiers.Add(TEXT("RW_Texture3D"));
-		ProcessedIdentifiers.Add(TEXT("RW_TextureCube"));
-		ProcessedIdentifiers.Add(TEXT("Texture1D_Array"));
-		ProcessedIdentifiers.Add(TEXT("Texture2D_Array"));
-		ProcessedIdentifiers.Add(TEXT("TextureBuffer"));
-		ProcessedIdentifiers.Add(TEXT("TextureCube_Array"));
+		ProcessedIdentifiers.Add(SHADER_SOURCE_LITERAL("AppendRegularBuffer"));
+		ProcessedIdentifiers.Add(SHADER_SOURCE_LITERAL("ByteBuffer"));
+		ProcessedIdentifiers.Add(SHADER_SOURCE_LITERAL("ConsumeRegularBuffer"));
+		ProcessedIdentifiers.Add(SHADER_SOURCE_LITERAL("DataBuffer"));
+		ProcessedIdentifiers.Add(SHADER_SOURCE_LITERAL("MS_Texture2D"));
+		ProcessedIdentifiers.Add(SHADER_SOURCE_LITERAL("MS_Texture2D_Array"));
+		ProcessedIdentifiers.Add(SHADER_SOURCE_LITERAL("RegularBuffer"));
+		ProcessedIdentifiers.Add(SHADER_SOURCE_LITERAL("RW_ByteBuffer"));
+		ProcessedIdentifiers.Add(SHADER_SOURCE_LITERAL("RW_DataBuffer"));
+		ProcessedIdentifiers.Add(SHADER_SOURCE_LITERAL("RW_RegularBuffer"));
+		ProcessedIdentifiers.Add(SHADER_SOURCE_LITERAL("RW_Texture1D"));
+		ProcessedIdentifiers.Add(SHADER_SOURCE_LITERAL("RW_Texture1D_Array"));
+		ProcessedIdentifiers.Add(SHADER_SOURCE_LITERAL("RW_Texture2D"));
+		ProcessedIdentifiers.Add(SHADER_SOURCE_LITERAL("RW_Texture2D_Array"));
+		ProcessedIdentifiers.Add(SHADER_SOURCE_LITERAL("RW_Texture3D"));
+		ProcessedIdentifiers.Add(SHADER_SOURCE_LITERAL("RW_TextureCube"));
+		ProcessedIdentifiers.Add(SHADER_SOURCE_LITERAL("Texture1D_Array"));
+		ProcessedIdentifiers.Add(SHADER_SOURCE_LITERAL("Texture2D_Array"));
+		ProcessedIdentifiers.Add(SHADER_SOURCE_LITERAL("TextureBuffer"));
+		ProcessedIdentifiers.Add(SHADER_SOURCE_LITERAL("TextureCube_Array"));
 
 		// Some shaders define template versions of some built-in functions, so we can't trivially ignore them
-		//ProcessedIdentifiers.Add(TEXT("abs"));
-		//ProcessedIdentifiers.Add(TEXT("any"));
-		//ProcessedIdentifiers.Add(TEXT("clamp"));
-		//ProcessedIdentifiers.Add(TEXT("clip"));
-		//ProcessedIdentifiers.Add(TEXT("cos"));
-		//ProcessedIdentifiers.Add(TEXT("cross"));
-		//ProcessedIdentifiers.Add(TEXT("dot"));
-		//ProcessedIdentifiers.Add(TEXT("frac"));
-		//ProcessedIdentifiers.Add(TEXT("lerp"));
-		//ProcessedIdentifiers.Add(TEXT("max"));
-		//ProcessedIdentifiers.Add(TEXT("min"));
-		//ProcessedIdentifiers.Add(TEXT("mul"));
-		//ProcessedIdentifiers.Add(TEXT("normalize"));
-		//ProcessedIdentifiers.Add(TEXT("pow"));
-		//ProcessedIdentifiers.Add(TEXT("saturate"));
-		//ProcessedIdentifiers.Add(TEXT("sign"));
-		//ProcessedIdentifiers.Add(TEXT("sin"));
-		//ProcessedIdentifiers.Add(TEXT("sqrt"));
+		//ProcessedIdentifiers.Add(SHADER_SOURCE_LITERAL("abs"));
+		//ProcessedIdentifiers.Add(SHADER_SOURCE_LITERAL("any"));
+		//ProcessedIdentifiers.Add(SHADER_SOURCE_LITERAL("clamp"));
+		//ProcessedIdentifiers.Add(SHADER_SOURCE_LITERAL("clip"));
+		//ProcessedIdentifiers.Add(SHADER_SOURCE_LITERAL("cos"));
+		//ProcessedIdentifiers.Add(SHADER_SOURCE_LITERAL("cross"));
+		//ProcessedIdentifiers.Add(SHADER_SOURCE_LITERAL("dot"));
+		//ProcessedIdentifiers.Add(SHADER_SOURCE_LITERAL("frac"));
+		//ProcessedIdentifiers.Add(SHADER_SOURCE_LITERAL("lerp"));
+		//ProcessedIdentifiers.Add(SHADER_SOURCE_LITERAL("max"));
+		//ProcessedIdentifiers.Add(SHADER_SOURCE_LITERAL("min"));
+		//ProcessedIdentifiers.Add(SHADER_SOURCE_LITERAL("mul"));
+		//ProcessedIdentifiers.Add(SHADER_SOURCE_LITERAL("normalize"));
+		//ProcessedIdentifiers.Add(SHADER_SOURCE_LITERAL("pow"));
+		//ProcessedIdentifiers.Add(SHADER_SOURCE_LITERAL("saturate"));
+		//ProcessedIdentifiers.Add(SHADER_SOURCE_LITERAL("sign"));
+		//ProcessedIdentifiers.Add(SHADER_SOURCE_LITERAL("sin"));
+		//ProcessedIdentifiers.Add(SHADER_SOURCE_LITERAL("sqrt"));
 	}
 
 	if (PendingChunks.IsEmpty())
@@ -1754,9 +1763,9 @@ static FString MinifyShader(const FParsedShader& Parsed, TConstArrayView<FString
 		return {};
 	}
 
-	TArray<FStringView> TempIdentifiers;
+	TArray<FShaderSource::FViewType> TempIdentifiers;
 
-	TMap<FStringView, TArray<const FCodeChunk*>, FDefaultSetAllocator, FCasedStringViewKeyFuncs> ChunksByIdentifier;
+	TMap<FShaderSource::FViewType, TArray<const FCodeChunk*>, FDefaultSetAllocator, FCasedStringViewKeyFuncs> ChunksByIdentifier;
 	for (const FCodeChunk& Chunk : Parsed.Chunks)
 	{
 		for (const FCodeBlock& Block : Chunk.Blocks)
@@ -1785,7 +1794,7 @@ static FString MinifyShader(const FParsedShader& Parsed, TConstArrayView<FString
 			{
 				TempIdentifiers.Reset();
 				ExtractIdentifiers(Block, TempIdentifiers);
-				for (FStringView Identifier : TempIdentifiers)
+				for (FShaderSource::FViewType Identifier : TempIdentifiers)
 				{
 					ChunksByIdentifier.FindOrAdd(Identifier).Push(&Chunk);
 				}
@@ -1815,7 +1824,7 @@ static FString MinifyShader(const FParsedShader& Parsed, TConstArrayView<FString
 
 		ExtractIdentifiers(*CurrentChunk, TempIdentifiers);
 
-		for (FStringView Identifier : TempIdentifiers)
+		for (FShaderSource::FViewType Identifier : TempIdentifiers)
 		{
 			bool bIdentifierWasAlreadyInSet = false;
 			ProcessedIdentifiers.Add(Identifier, &bIdentifierWasAlreadyInSet);
@@ -1889,35 +1898,35 @@ static FString MinifyShader(const FParsedShader& Parsed, TConstArrayView<FString
 
 	if (EnumHasAnyFlags(Flags, EMinifyShaderFlags::OutputStats))
 	{
-		OutputStream.Append(TEXTVIEW("// Total code chunks: "));
+		OutputStream.Append(SHADER_SOURCE_VIEWLITERAL("// Total code chunks: "));
 		OutputStream.AppendInt(RelevantChunks.Num());
-		OutputStream.AppendChar(L'\n');
+		OutputStream.AppendChar('\n');
 
-		OutputStream.Append(TEXTVIEW("// - Functions: "));
+		OutputStream.Append(SHADER_SOURCE_VIEWLITERAL("// - Functions: "));
 		OutputStream.AppendInt(NumFunctions);
-		OutputStream.AppendChar(L'\n');
+		OutputStream.AppendChar('\n');
 
-		OutputStream.Append(TEXTVIEW("// - Structs: "));
+		OutputStream.Append(SHADER_SOURCE_VIEWLITERAL("// - Structs: "));
 		OutputStream.AppendInt(NumStructs);
-		OutputStream.AppendChar(L'\n');
+		OutputStream.AppendChar('\n');
 
-		OutputStream.Append(TEXTVIEW("// - CBuffers: "));
+		OutputStream.Append(SHADER_SOURCE_VIEWLITERAL("// - CBuffers: "));
 		OutputStream.AppendInt(NumCBuffers);
-		OutputStream.AppendChar(L'\n');
+		OutputStream.AppendChar('\n');
 
-		OutputStream.Append(TEXTVIEW("// - Variables: "));
+		OutputStream.Append(SHADER_SOURCE_VIEWLITERAL("// - Variables: "));
 		OutputStream.AppendInt(NumVariables);
-		OutputStream.AppendChar(L'\n');
+		OutputStream.AppendChar('\n');
 
-		OutputStream.Append(TEXTVIEW("// - Other: "));
+		OutputStream.Append(SHADER_SOURCE_VIEWLITERAL("// - Other: "));
 		OutputStream.AppendInt(NumOtherChunks);
-		OutputStream.AppendChar(L'\n');
+		OutputStream.AppendChar('\n');
 
-		OutputStream.AppendChar(L'\n');
+		OutputStream.AppendChar('\n');
 	}
 
 	TArray<int32> LineBreakMap;
-	const TArray<FStringView>& LineDirectives = Parsed.LineDirectives;
+	const TArray<FShaderSource::FViewType>& LineDirectives = Parsed.LineDirectives;
 	if (EnumHasAnyFlags(Flags, EMinifyShaderFlags::OutputLines))
 	{
 		BuildLineBreakMap(Parsed.Source, LineBreakMap);
@@ -1926,7 +1935,7 @@ static FString MinifyShader(const FParsedShader& Parsed, TConstArrayView<FString
 	const FNamespace* CurrentNamespace = nullptr;
 
 	int32 LastLineNumber = -1;
-	FStringView LastLineFileName;
+	FShaderSource::FViewType LastLineFileName;
 
 	for (const FCodeChunk& Chunk : Parsed.Chunks)
 	{
@@ -1970,13 +1979,13 @@ static FString MinifyShader(const FParsedShader& Parsed, TConstArrayView<FString
 			if (CurrentNamespace)
 			{
 				CloseNamespace(OutputStream, *CurrentNamespace);
-				OutputStream.Append(TEXTVIEW("\n\n"));
+				OutputStream.Append(SHADER_SOURCE_VIEWLITERAL("\n\n"));
 			}
 
 			if (PendingNamespace)
 			{
 				OpenNamespace(OutputStream, *PendingNamespace);
-				OutputStream.Append(TEXTVIEW("\n\n"));
+				OutputStream.Append(SHADER_SOURCE_VIEWLITERAL("\n\n"));
 			}
 
 			CurrentNamespace = PendingNamespace;
@@ -1988,19 +1997,19 @@ static FString MinifyShader(const FParsedShader& Parsed, TConstArrayView<FString
 			if (RequestedBy != nullptr)
 			{
 				const FCodeChunk* RequestedByChunk = *RequestedBy;
-				FStringView  RequestedByName = RequestedByChunk->FindFirstBlockByType(EBlockType::Name);
+				FShaderSource::FViewType  RequestedByName = RequestedByChunk->FindFirstBlockByType(EBlockType::Name);
 				if (!RequestedByName.IsEmpty())
 				{
-					OutputStream.Append(TEXTVIEW("// REASON: "));
+					OutputStream.Append(SHADER_SOURCE_VIEWLITERAL("// REASON: "));
 					OutputStream.Append(RequestedByName);
-					OutputStream.AppendChar(L'\n');
+					OutputStream.AppendChar('\n');
 				}
 			}
 		}
 
 		if (EnumHasAnyFlags(Flags, EMinifyShaderFlags::OutputLines))
 		{
-			const FStringView ChunkCode = Chunk.Blocks[0];
+			const FShaderSource::FViewType ChunkCode = Chunk.Blocks[0];
 			int32 LineDirectiveIndex = FindLineDirective(LineDirectives, ChunkCode.GetData());
 			int32 ChunkLine = FindLineNumber(Parsed.Source, LineBreakMap, ChunkCode.GetData());
 
@@ -2009,13 +2018,13 @@ static FString MinifyShader(const FParsedShader& Parsed, TConstArrayView<FString
 				// There was no valid line directive for this chunk, but we do know the line in the input source, so just emit that.
 				if (ChunkLine > LastLineNumber + 1 || !LastLineFileName.IsEmpty())
 				{
-					OutputStream.Append(TEXTVIEW("#line "));
+					OutputStream.Append(SHADER_SOURCE_VIEWLITERAL("#line "));
 					OutputStream.AppendInt(ChunkLine);
-					OutputStream.AppendChar(L'\n');
+					OutputStream.AppendChar('\n');
 				}
 
 				LastLineNumber = ChunkLine;
-				LastLineFileName = FStringView();
+				LastLineFileName = FShaderSource::FViewType();
 			}
 			else if (ChunkLine != INDEX_NONE && LineDirectiveIndex != INDEX_NONE)
 			{
@@ -2023,10 +2032,10 @@ static FString MinifyShader(const FParsedShader& Parsed, TConstArrayView<FString
 				// Some of the input source code may have been removed, so we need to adjust 
 				// the line number before emitting the line directive.
 
-				FStringView LineDirective = LineDirectives[LineDirectiveIndex];
+				FShaderSource::FViewType LineDirective = LineDirectives[LineDirectiveIndex];
 				int32 LineDirectiveLine = FindLineNumber(Parsed.Source, LineBreakMap, LineDirective.GetData());
 				int32 ParsedLineNumber = INDEX_NONE;
-				FStringView ParsedFileName;
+				FShaderSource::FViewType ParsedFileName;
 				if (LineDirectiveLine != INDEX_NONE && ParseLineDirective(LineDirective, ParsedLineNumber, ParsedFileName))
 				{
 					int32 OffsetFromLineDirective = ChunkLine - (LineDirectiveLine + 1); // Line directive identifies the *next* line, hence +1 when computing the offset
@@ -2037,19 +2046,19 @@ static FString MinifyShader(const FParsedShader& Parsed, TConstArrayView<FString
 						// Separate the next block from the previous one when it starts with a line directive
 						if (OutputStream.Len())
 						{
-							OutputStream.AppendChar(L'\n');
+							OutputStream.AppendChar('\n');
 						}
 
-						OutputStream.Append(TEXTVIEW("#line "));
+						OutputStream.Append(SHADER_SOURCE_VIEWLITERAL("#line "));
 						OutputStream.AppendInt(PatchedLineNumber);
 
 						if (!ParsedFileName.IsEmpty())
 						{
-							OutputStream.Append(TEXTVIEW(" \""));
+							OutputStream.Append(SHADER_SOURCE_VIEWLITERAL(" \""));
 							OutputStream.Append(ParsedFileName);
-							OutputStream.AppendChar(L'\"');
+							OutputStream.AppendChar('\"');
 						}
-						OutputStream.AppendChar(L'\n');
+						OutputStream.AppendChar('\n');
 					}
 
 					LastLineNumber = PatchedLineNumber;
@@ -2064,24 +2073,24 @@ static FString MinifyShader(const FParsedShader& Parsed, TConstArrayView<FString
 	if (CurrentNamespace)
 	{
 		CloseNamespace(OutputStream, *CurrentNamespace);
-		OutputStream.AppendChar(L'\n');
+		OutputStream.AppendChar('\n');
 		CurrentNamespace = nullptr;
 	}
 
 	return OutputStream;
 }
 
-static FString MinifyShader(const FParsedShader& Parsed, FStringView EntryPoint, EMinifyShaderFlags Flags, FDiagnostics& Diagnostics)
+static FShaderSource::FStringType MinifyShader(const FParsedShader& Parsed, FShaderSource::FViewType EntryPoint, EMinifyShaderFlags Flags, FDiagnostics& Diagnostics)
 {
-	TArray<FStringView> RequiredSymbols = SplitByChar(EntryPoint, ';');
+	TArray<FShaderSource::FViewType> RequiredSymbols = SplitByChar(EntryPoint, ';');
 	return MinifyShader(Parsed, RequiredSymbols, Flags, Diagnostics);
 }
 
-FMinifiedShader Minify(const FStringView PreprocessedShader, TConstArrayView<FStringView> RequiredSymbols, EMinifyShaderFlags Flags)
+FMinifiedShader Minify(const FShaderSource& PreprocessedShader, TConstArrayView<FShaderSource::FViewType> RequiredSymbols, EMinifyShaderFlags Flags)
 {
 	FMinifiedShader Result;
 
-	FParsedShader Parsed = ParseShader(PreprocessedShader, Result.Diagnostics);
+	FParsedShader Parsed = ParseShader(PreprocessedShader.GetView(), Result.Diagnostics);
 
 	if (!Parsed.Chunks.IsEmpty())
 	{
@@ -2091,7 +2100,7 @@ FMinifiedShader Minify(const FStringView PreprocessedShader, TConstArrayView<FSt
 	return Result;
 }
 
-FMinifiedShader Minify(const FStringView PreprocessedShader, const FStringView EntryPoint, EMinifyShaderFlags Flags)
+FMinifiedShader Minify(const FShaderSource& PreprocessedShader, const FShaderSource::FViewType EntryPoint, EMinifyShaderFlags Flags)
 {
 	return Minify(PreprocessedShader, MakeArrayView(&EntryPoint, 1), Flags);
 }
@@ -2105,7 +2114,7 @@ IMPLEMENT_SIMPLE_AUTOMATION_TEST(FShaderMinifierParserTest, "System.Shaders.Shad
 namespace UE::ShaderMinifier
 {
 // Convenience wrapper for tests where we don't care about diagnostic messages
-static FParsedShader ParseShader(FStringView InSource)
+static FParsedShader ParseShader(FShaderSource::FViewType InSource)
 {
 	FDiagnostics Diagnostics;
 	return ParseShader(InSource, Diagnostics);
@@ -2117,19 +2126,19 @@ bool FShaderMinifierParserTest::RunTest(const FString& Parameters)
 	using namespace UE::ShaderMinifier;
 
 	TestEqual(TEXT("SkipSpace"), 
-		FString(SkipSpace(TEXT("  \n\r\f \tHello"))), 
-		FString(TEXT("Hello")));
+		FString(SkipSpace(SHADER_SOURCE_LITERAL("  \n\r\f \tHello"))), 
+		FString(SHADER_SOURCE_LITERAL("Hello")));
 
 	TestEqual(TEXT("SkipUntilStr (found)"), 
-		FString(SkipUntilStr(TEXT("Hello World"), TEXT("World"))),
-		FString(TEXT("World")));
+		FString(SkipUntilStr(SHADER_SOURCE_LITERAL("Hello World"), SHADER_SOURCE_LITERAL("World"))),
+		FString(SHADER_SOURCE_LITERAL("World")));
 
 	TestEqual(TEXT("SkipUntilStr (not found)"),
-		FString(SkipUntilStr(TEXT("Hello World"), TEXT("Blah"))),
+		FString(SkipUntilStr(SHADER_SOURCE_LITERAL("Hello World"), SHADER_SOURCE_LITERAL("Blah"))),
 		FString());
 
 	{
-		auto P = ParseShader(TEXT("static const struct { int Blah; } Foo = { 123; };"));
+		auto P = ParseShader(SHADER_SOURCE_LITERAL("static const struct { int Blah; } Foo = { 123; };"));
 		TestEqual(TEXT("Anonymous struct variable with initializer, total chunks"), P.Chunks.Num(), 1);
 		if (P.Chunks.Num() == 1)
 		{
@@ -2138,7 +2147,7 @@ bool FShaderMinifierParserTest::RunTest(const FString& Parameters)
 	}
 
 	{
-		auto P = ParseShader(TEXT("float4 PSMain() : SV_Target { return float4(1,0,0,1); };"));
+		auto P = ParseShader(SHADER_SOURCE_LITERAL("float4 PSMain() : SV_Target { return float4(1,0,0,1); };"));
 		TestEqual(TEXT("Pixel shader entry point, total chunks"), P.Chunks.Num(), 1);
 		if (P.Chunks.Num() == 1)
 		{
@@ -2147,32 +2156,32 @@ bool FShaderMinifierParserTest::RunTest(const FString& Parameters)
 	}
 
 	{
-		TArray<FStringView> R;
-		ExtractIdentifiers(TEXT("Hello[World]; Foo[0];\n"), R);
+		TArray<FShaderSource::FViewType> R;
+		ExtractIdentifiers(SHADER_SOURCE_LITERAL("Hello[World]; Foo[0];\n"), R);
 		if (TestEqual(TEXT("ExtractIdentifiers1: Num"), R.Num(), 3))
 		{
-			TestEqual(TEXT("ExtractIdentifiers1: R[0]"), FString(R[0]), TEXT("Hello"));
-			TestEqual(TEXT("ExtractIdentifiers1: R[1]"), FString(R[1]), TEXT("World"));
-			TestEqual(TEXT("ExtractIdentifiers1: R[2]"), FString(R[2]), TEXT("Foo"));
+			TestEqual(TEXT("ExtractIdentifiers1: R[0]"), FString(R[0]), SHADER_SOURCE_LITERAL("Hello"));
+			TestEqual(TEXT("ExtractIdentifiers1: R[1]"), FString(R[1]), SHADER_SOURCE_LITERAL("World"));
+			TestEqual(TEXT("ExtractIdentifiers1: R[2]"), FString(R[2]), SHADER_SOURCE_LITERAL("Foo"));
 		}
 	}
 
 	{
-		TArray<FStringView> R;
-		ExtractIdentifiers(TEXT("#line 0\nStructuredBuffer<uint4> Blah : register(t0, space123);#line 1\n#pragma foo\n"), R);
+		TArray<FShaderSource::FViewType> R;
+		ExtractIdentifiers(SHADER_SOURCE_LITERAL("#line 0\nStructuredBuffer<uint4> Blah : register(t0, space123);#line 1\n#pragma foo\n"), R);
 		if (TestEqual(TEXT("ExtractIdentifiers2: Num"), R.Num(), 6))
 		{
-			TestEqual(TEXT("ExtractIdentifiers2: R[0]"), FString(R[0]), TEXT("StructuredBuffer"));
-			TestEqual(TEXT("ExtractIdentifiers2: R[1]"), FString(R[1]), TEXT("uint4"));
-			TestEqual(TEXT("ExtractIdentifiers2: R[2]"), FString(R[2]), TEXT("Blah"));
-			TestEqual(TEXT("ExtractIdentifiers2: R[3]"), FString(R[3]), TEXT("register"));
-			TestEqual(TEXT("ExtractIdentifiers2: R[4]"), FString(R[4]), TEXT("t0"));
-			TestEqual(TEXT("ExtractIdentifiers2: R[5]"), FString(R[5]), TEXT("space123"));
+			TestEqual(TEXT("ExtractIdentifiers2: R[0]"), FString(R[0]), SHADER_SOURCE_LITERAL("StructuredBuffer"));
+			TestEqual(TEXT("ExtractIdentifiers2: R[1]"), FString(R[1]), SHADER_SOURCE_LITERAL("uint4"));
+			TestEqual(TEXT("ExtractIdentifiers2: R[2]"), FString(R[2]), SHADER_SOURCE_LITERAL("Blah"));
+			TestEqual(TEXT("ExtractIdentifiers2: R[3]"), FString(R[3]), SHADER_SOURCE_LITERAL("register"));
+			TestEqual(TEXT("ExtractIdentifiers2: R[4]"), FString(R[4]), SHADER_SOURCE_LITERAL("t0"));
+			TestEqual(TEXT("ExtractIdentifiers2: R[5]"), FString(R[5]), SHADER_SOURCE_LITERAL("space123"));
 		}
 	}
 
 	{
-		auto P = ParseShader(TEXT("StructuredBuffer<uint4> Blah : register(t0, space123);"));
+		auto P = ParseShader(SHADER_SOURCE_LITERAL("StructuredBuffer<uint4> Blah : register(t0, space123);"));
 		if (TestEqual(TEXT("ParseShader: structured buffer: num chunks"), P.Chunks.Num(), 1))
 		{
 			TestEqual(TEXT("ParseShader: structured buffer: chunk"), P.Chunks[0].Type, ECodeChunkType::Variable);
@@ -2180,7 +2189,7 @@ bool FShaderMinifierParserTest::RunTest(const FString& Parameters)
 	}
 
 	{
-		auto P = ParseShader(TEXT("const float Foo = 123.45f;"));
+		auto P = ParseShader(SHADER_SOURCE_LITERAL("const float Foo = 123.45f;"));
 		if (TestEqual(TEXT("ParseShader: const float with initializer: num chunks"), P.Chunks.Num(), 1))
 		{
 			TestEqual(TEXT("ParseShader: const float with initializer: chunk type"), P.Chunks[0].Type, ECodeChunkType::Variable);
@@ -2188,7 +2197,7 @@ bool FShaderMinifierParserTest::RunTest(const FString& Parameters)
 	}
 
 	{
-		auto P = ParseShader(TEXT("struct Blah { int A; };"));
+		auto P = ParseShader(SHADER_SOURCE_LITERAL("struct Blah { int A; };"));
 		if (TestEqual(TEXT("ParseShader: struct: num chunks"), P.Chunks.Num(), 1))
 		{
 			TestEqual(TEXT("ParseShader: struct: chunk type"), P.Chunks[0].Type, ECodeChunkType::Struct);
@@ -2196,7 +2205,7 @@ bool FShaderMinifierParserTest::RunTest(const FString& Parameters)
 	}
 
 	{
-		auto P = ParseShader(TEXT("struct Foo { int FooA; }; struct Bar : Foo { int BarA; };"));
+		auto P = ParseShader(SHADER_SOURCE_LITERAL("struct Foo { int FooA; }; struct Bar : Foo { int BarA; };"));
 		if (TestEqual(TEXT("ParseShader: inherited struct: num chunks"), P.Chunks.Num(), 2))
 		{
 			TestEqual(TEXT("ParseShader: inherited struct: chunk 0 type"), P.Chunks[0].Type, ECodeChunkType::Struct);
@@ -2205,7 +2214,7 @@ bool FShaderMinifierParserTest::RunTest(const FString& Parameters)
 	}
 
 	{
-		auto P = ParseShader(TEXT("[numthreads(8,8,1)] void Main() {};"));
+		auto P = ParseShader(SHADER_SOURCE_LITERAL("[numthreads(8,8,1)] void Main() {};"));
 		if (TestEqual(TEXT("ParseShader: compute shader entry point: num chunks"), P.Chunks.Num(), 1))
 		{
 			TestEqual(TEXT("ParseShader: compute shader entry point: chunk type"), P.Chunks[0].Type, ECodeChunkType::Function);
@@ -2218,14 +2227,14 @@ bool FShaderMinifierParserTest::RunTest(const FString& Parameters)
 	}
 
 	{
-		auto P = ParseShader(TEXT("Texture2D Blah : register(t0);"));
+		auto P = ParseShader(SHADER_SOURCE_LITERAL("Texture2D Blah : register(t0);"));
 		if (TestEqual(TEXT("ParseShader: texture with register: num chunks"), P.Chunks.Num(), 1))
 		{
 			TestEqual(TEXT("ParseShader: texture with register: chunk type"), P.Chunks[0].Type, ECodeChunkType::Variable);
 		}
 	}
 	{
-		auto P = ParseShader(TEXT("Texture2D Blah;"));
+		auto P = ParseShader(SHADER_SOURCE_LITERAL("Texture2D Blah;"));
 		if (TestEqual(TEXT("ParseShader: texture: num chunks"), P.Chunks.Num(), 1))
 		{
 			TestEqual(TEXT("ParseShader: texture: chunk type"), P.Chunks[0].Type, ECodeChunkType::Variable);
@@ -2233,7 +2242,7 @@ bool FShaderMinifierParserTest::RunTest(const FString& Parameters)
 	}
 
 	{
-		auto P = ParseShader(TEXT("SamplerState Blah : register(s0, space123);"));
+		auto P = ParseShader(SHADER_SOURCE_LITERAL("SamplerState Blah : register(s0, space123);"));
 		if (TestEqual(TEXT("ParseShader: sampler state with register: num chunks"), P.Chunks.Num(), 1))
 		{
 			TestEqual(TEXT("ParseShader: sampler state with register: chunk type"), P.Chunks[0].Type, ECodeChunkType::Variable);
@@ -2243,13 +2252,13 @@ bool FShaderMinifierParserTest::RunTest(const FString& Parameters)
 #if 0
 	{
 		// TODO: handle function forward declarations
-		auto P = ParseShader(TEXT("Foo Fun(int a);"));
+		auto P = ParseShader(SHADER_SOURCE_LITERAL("Foo Fun(int a);"));
 		TestEqual(TEXT("ParseShader: function forward declaration"), P.Chunks[0].Type, ECodeChunkType::FunctionDecl);
 	}
 #endif
 
 	{
-		auto P = ParseShader(TEXT("void Fun(int a) {};"));
+		auto P = ParseShader(SHADER_SOURCE_LITERAL("void Fun(int a) {};"));
 		if (TestEqual(TEXT("ParseShader: function with trailing semicolon: num chunks"), P.Chunks.Num(), 1))
 		{
 			TestEqual(TEXT("ParseShader: function with trailing semicolon: chunk type"), P.Chunks[0].Type, ECodeChunkType::Function);
@@ -2257,7 +2266,7 @@ bool FShaderMinifierParserTest::RunTest(const FString& Parameters)
 	}
 
 	{
-		auto P = ParseShader(TEXT("void Fun(int a) {}"));
+		auto P = ParseShader(SHADER_SOURCE_LITERAL("void Fun(int a) {}"));
 		if (TestEqual(TEXT("ParseShader: function: num chunks"), P.Chunks.Num(), 1))
 		{
 			TestEqual(TEXT("ParseShader: function: chunk type"), P.Chunks[0].Type, ECodeChunkType::Function);
@@ -2265,7 +2274,7 @@ bool FShaderMinifierParserTest::RunTest(const FString& Parameters)
 	}
 
 	{
-		auto P = ParseShader(TEXT("cbuffer Foo {blah} SamplerState S;"));
+		auto P = ParseShader(SHADER_SOURCE_LITERAL("cbuffer Foo {blah} SamplerState S;"));
 
 		if (TestEqual(TEXT("ParseShader: cbuffer and sampler state: num chunks"), P.Chunks.Num(), 2))
 		{
@@ -2275,7 +2284,7 @@ bool FShaderMinifierParserTest::RunTest(const FString& Parameters)
 	}
 
 	{
-		auto P = ParseShader(TEXT("struct Foo { int a; };"));
+		auto P = ParseShader(SHADER_SOURCE_LITERAL("struct Foo { int a; };"));
 		if (TestEqual(TEXT("ParseShader: struct: num chunks"), P.Chunks.Num(), 1))
 		{
 			TestEqual(TEXT("ParseShader: struct: chunk type"), P.Chunks[0].Type, ECodeChunkType::Struct);
@@ -2283,7 +2292,7 @@ bool FShaderMinifierParserTest::RunTest(const FString& Parameters)
 	}
 
 	{
-		auto P = ParseShader(TEXT("struct { int a; } Foo = { 123; };"));
+		auto P = ParseShader(SHADER_SOURCE_LITERAL("struct { int a; } Foo = { 123; };"));
 		if (TestEqual(TEXT("ParseShader: anonymous struct with variable and initializer: num chunks"), P.Chunks.Num(), 1))
 		{
 			TestEqual(TEXT("ParseShader: anonymous struct with variable and initializer: chunk type [0]"), P.Chunks[0].Type, ECodeChunkType::Variable);
@@ -2293,14 +2302,14 @@ bool FShaderMinifierParserTest::RunTest(const FString& Parameters)
 #if 0
 	{
 		// TODO: handle struct forward declarations
-		auto P = ParseShader(TEXT("struct Foo;"));
+		auto P = ParseShader(SHADER_SOURCE_LITERAL("struct Foo;"));
 	}
 #endif
 
 	{
 		auto P = ParseShader(
-			TEXT("cbuffer MyBuffer : register(b3)")
-			TEXT("{ float4 Element1 : packoffset(c0); float1 Element2 : packoffset(c1); float1 Element3 : packoffset(c1.y); }"));
+			SHADER_SOURCE_LITERAL("cbuffer MyBuffer : register(b3)")
+			SHADER_SOURCE_LITERAL("{ float4 Element1 : packoffset(c0); float1 Element2 : packoffset(c1); float1 Element3 : packoffset(c1.y); }"));
 		if (TestEqual(TEXT("ParseShader: cbuffer with packoffset: num chunks"), P.Chunks.Num(), 1))
 		{
 			TestEqual(TEXT("ParseShader: cbuffer with packoffset: chunk type"), P.Chunks[0].Type, ECodeChunkType::CBuffer);
@@ -2308,7 +2317,7 @@ bool FShaderMinifierParserTest::RunTest(const FString& Parameters)
 	}
 
 	{
-		auto P = ParseShader(TEXT("static const struct { float4 Param; } Foo;"));
+		auto P = ParseShader(SHADER_SOURCE_LITERAL("static const struct { float4 Param; } Foo;"));
 		if (TestEqual(TEXT("ParseShader: static const anonymous struct with variable: num chunks"), P.Chunks.Num(), 1))
 		{
 			TestEqual(TEXT("ParseShader: static const anonymous struct with variable: chunk type"), P.Chunks[0].Type, ECodeChunkType::Variable);
@@ -2316,7 +2325,7 @@ bool FShaderMinifierParserTest::RunTest(const FString& Parameters)
 	}
 
 	{
-		auto P = ParseShader(TEXT("static const struct { float4 Param; } Foo = { FooCB_Param; };"));
+		auto P = ParseShader(SHADER_SOURCE_LITERAL("static const struct { float4 Param; } Foo = { FooCB_Param; };"));
 		if (TestEqual(TEXT("ParseShader: static const anonymous struct with variable and initializer: num chunks"), P.Chunks.Num(), 1))
 		{
 			TestEqual(TEXT("ParseShader: static const anonymous struct with variable and initializer: chunk type"), P.Chunks[0].Type, ECodeChunkType::Variable);
@@ -2324,7 +2333,7 @@ bool FShaderMinifierParserTest::RunTest(const FString& Parameters)
 	}
 
 	{
-		auto P = ParseShader(TEXT("template <typename T> float Fun(T x) { return (float)x; }"));
+		auto P = ParseShader(SHADER_SOURCE_LITERAL("template <typename T> float Fun(T x) { return (float)x; }"));
 		if (TestEqual(TEXT("ParseShader: template function: num chunks"), P.Chunks.Num(), 1))
 		{
 			TestEqual(TEXT("ParseShader: template function: chunk type"), P.Chunks[0].Type, ECodeChunkType::Function);
@@ -2332,7 +2341,7 @@ bool FShaderMinifierParserTest::RunTest(const FString& Parameters)
 	}
 
 	{
-		auto P = ParseShader(TEXT("enum EFoo { A, B = 123 };"));
+		auto P = ParseShader(SHADER_SOURCE_LITERAL("enum EFoo { A, B = 123 };"));
 		if (TestEqual(TEXT("ParseShader: enum: num chunks"), P.Chunks.Num(), 1))
 		{
 			TestEqual(TEXT("ParseShader: enum: chunk type"), P.Chunks[0].Type, ECodeChunkType::Enum);
@@ -2340,7 +2349,7 @@ bool FShaderMinifierParserTest::RunTest(const FString& Parameters)
 	}
 
 	{
-		auto P = ParseShader(TEXT("enum class EFoo { A, B };"));
+		auto P = ParseShader(SHADER_SOURCE_LITERAL("enum class EFoo { A, B };"));
 		if (TestEqual(TEXT("ParseShader: enum class: num chunks"), P.Chunks.Num(), 1))
 		{
 			TestEqual(TEXT("ParseShader: enum class: chunk type"), P.Chunks[0].Type, ECodeChunkType::Enum);
@@ -2348,7 +2357,7 @@ bool FShaderMinifierParserTest::RunTest(const FString& Parameters)
 	}
 
 	{
-		auto P = ParseShader(TEXT("#define Foo 123"));
+		auto P = ParseShader(SHADER_SOURCE_LITERAL("#define Foo 123"));
 		if (TestEqual(TEXT("ParseShader: define: num chunks"), P.Chunks.Num(), 1))
 		{
 			TestEqual(TEXT("ParseShader: define: chunk type"), P.Chunks[0].Type, ECodeChunkType::Define);
@@ -2356,7 +2365,7 @@ bool FShaderMinifierParserTest::RunTest(const FString& Parameters)
 	}
 
 	{
-		auto P = ParseShader(TEXT("#pragma Foo"));
+		auto P = ParseShader(SHADER_SOURCE_LITERAL("#pragma Foo"));
 		if (TestEqual(TEXT("ParseShader: pragma: num chunks"), P.Chunks.Num(), 1))
 		{
 			TestEqual(TEXT("ParseShader: pragma: chunk type"), P.Chunks[0].Type, ECodeChunkType::Pragma);
@@ -2364,7 +2373,7 @@ bool FShaderMinifierParserTest::RunTest(const FString& Parameters)
 	}
 
 	{
-		auto P = ParseShader(TEXT("ConstantBuffer<Foo> CB : register ( b123, space456);"));
+		auto P = ParseShader(SHADER_SOURCE_LITERAL("ConstantBuffer<Foo> CB : register ( b123, space456);"));
 		if (TestEqual(TEXT("ParseShader: ConstantBuffer<Foo>: num chunks"), P.Chunks.Num(), 1))
 		{
 			TestEqual(TEXT("ParseShader: ConstantBuffer<Foo>: chunk type"), P.Chunks[0].Type, ECodeChunkType::Variable);
@@ -2372,7 +2381,7 @@ bool FShaderMinifierParserTest::RunTest(const FString& Parameters)
 	}
 
 	{
-		auto P = ParseShader(TEXT("namespace NS1 { void Fun() {}; } namespace NS2 { void Fun() {}; }"));
+		auto P = ParseShader(SHADER_SOURCE_LITERAL("namespace NS1 { void Fun() {}; } namespace NS2 { void Fun() {}; }"));
 		if (TestEqual(TEXT("ParseShader: namespaces: num chunks"), P.Chunks.Num(), 2)
 			&& TestEqual(TEXT("ParseShader: namespaces: num namespaces"), P.Namespaces.Num(), 2))
 		{
@@ -2382,26 +2391,26 @@ bool FShaderMinifierParserTest::RunTest(const FString& Parameters)
 	}
 
 	{
-		auto P = ParseShader(TEXT("template< typename T > TMyStruct<T> operator + ( TMyStruct<T> A, T B ) { /*...*/ }"));
+		auto P = ParseShader(SHADER_SOURCE_LITERAL("template< typename T > TMyStruct<T> operator + ( TMyStruct<T> A, T B ) { /*...*/ }"));
 		if (TestEqual(TEXT("ParseShader: operators: num chunks"), P.Chunks.Num(), 1))
 		{
 			TestEqual(TEXT("ParseShader: operators: chunk type"), P.Chunks[0].Type, ECodeChunkType::Operator);
 			if (TestEqual(TEXT("ParseShader: operators: chunk 0: num blocks"), P.Chunks[0].Blocks.Num(), 8))
 			{
 				FString Args = FString(P.Chunks[0].FindFirstBlockByType(EBlockType::Args));
-				TestEqual(TEXT("ParseShader: operators: chunk 0: type name"), *Args, TEXT("( TMyStruct<T> A, T B )"));
+				TestEqual(TEXT("ParseShader: operators: chunk 0: type name"), *Args, SHADER_SOURCE_LITERAL("( TMyStruct<T> A, T B )"));
 
 				FString TypeName = FString(P.Chunks[0].FindFirstBlockByType(EBlockType::Type));
-				TestEqual(TEXT("ParseShader: operators: chunk 0: type name"), *TypeName, TEXT("TMyStruct"));
+				TestEqual(TEXT("ParseShader: operators: chunk 0: type name"), *TypeName, SHADER_SOURCE_LITERAL("TMyStruct"));
 
 				FString OperatorName = FString(P.Chunks[0].FindFirstBlockByType(EBlockType::OperatorName));
-				TestEqual(TEXT("ParseShader: operators: chunk 0: operator name"), *OperatorName, TEXT("+"));
+				TestEqual(TEXT("ParseShader: operators: chunk 0: operator name"), *OperatorName, SHADER_SOURCE_LITERAL("+"));
 			}
 		}
 	}
 
 	{
-		auto P = ParseShader(TEXT("typedef Bar Foo;"));
+		auto P = ParseShader(SHADER_SOURCE_LITERAL("typedef Bar Foo;"));
 		if (TestEqual(TEXT("ParseShader: standard typedef: num chunks"), P.Chunks.Num(), 1))
 		{
 			TestEqual(TEXT("ParseShader: standard typedef: chunk type"), P.Chunks[0].Type, ECodeChunkType::Typedef);
@@ -2409,7 +2418,7 @@ bool FShaderMinifierParserTest::RunTest(const FString& Parameters)
 	}
 
 	{
-		auto P = ParseShader(TEXT("typedef Bar Foo; static const Foo = Bar(0);"));
+		auto P = ParseShader(SHADER_SOURCE_LITERAL("typedef Bar Foo; static const Foo = Bar(0);"));
 		if (TestEqual(TEXT("ParseShader: standard typedef: num chunks"), P.Chunks.Num(), 2))
 		{
 			TestEqual(TEXT("ParseShader: standard typedef: chunk type"), P.Chunks[0].Type, ECodeChunkType::Typedef);
@@ -2428,8 +2437,8 @@ bool FShaderMinifierTest::RunTest(const FString& Parameters)
 {
 	using namespace UE::ShaderMinifier;
 
-	FStringView TestShaderCode = 
-		TEXT(R"(// dxc /T cs_6_6 /E MainCS MinifierTest.hlsl 
+	FShaderSource::FViewType TestShaderCode = 
+		SHADER_SOURCE_LITERAL(R"(// dxc /T cs_6_6 /E MainCS MinifierTest.hlsl 
 struct FFoo
 {
 	float X;
@@ -2570,7 +2579,7 @@ void MainCS()
 }
 )");
 
-	auto ChunkPresent = [](const FParsedShader& Parsed, FStringView Name)
+	auto ChunkPresent = [](const FParsedShader& Parsed, FShaderSource::FViewType Name)
 	{
 		for (const FCodeChunk& Chunk : Parsed.Chunks)
 		{
@@ -2589,59 +2598,59 @@ void MainCS()
 
 	{
 		FDiagnostics Diagnostics;
-		FString Minified = MinifyShader(Parsed, TEXT("EmptyFunction"), EMinifyShaderFlags::None, Diagnostics);
+		FShaderSource::FStringType Minified = MinifyShader(Parsed, SHADER_SOURCE_LITERAL("EmptyFunction"), EMinifyShaderFlags::None, Diagnostics);
 		FParsedShader MinifiedParsed = ParseShader(Minified);
 		if (TestEqual(TEXT("MinifyShader: EmptyFunction: num chunks"), MinifiedParsed.Chunks.Num(), 3))
 		{
-			TestEqual(TEXT("MinifyShader: EmptyFunction: pragma"), *FString(MinifiedParsed.Chunks[0].GetCode()), TEXT("#pragma test_pragma"));
-			TestEqual(TEXT("MinifyShader: EmptyFunction: define"), *FString(MinifiedParsed.Chunks[1].GetCode()), TEXT("#define COMPILER_DEFINITION_TEST 123"));
-			TestEqual(TEXT("MinifyShader: EmptyFunction: function"), *FString(MinifiedParsed.Chunks[2].GetCode()), TEXT("void EmptyFunction(){}"));
+			TestEqual(TEXT("MinifyShader: EmptyFunction: pragma"), *FString(MinifiedParsed.Chunks[0].GetCode()), SHADER_SOURCE_LITERAL("#pragma test_pragma"));
+			TestEqual(TEXT("MinifyShader: EmptyFunction: define"), *FString(MinifiedParsed.Chunks[1].GetCode()), SHADER_SOURCE_LITERAL("#define COMPILER_DEFINITION_TEST 123"));
+			TestEqual(TEXT("MinifyShader: EmptyFunction: function"), *FString(MinifiedParsed.Chunks[2].GetCode()), SHADER_SOURCE_LITERAL("void EmptyFunction(){}"));
 		}
 	}
 
 	{
 		FDiagnostics Diagnostics;
-		FString Minified = MinifyShader(Parsed, TEXT("MainCS"), EMinifyShaderFlags::OutputReasons, Diagnostics);
+		FShaderSource::FStringType Minified = MinifyShader(Parsed, SHADER_SOURCE_LITERAL("MainCS"), EMinifyShaderFlags::OutputReasons, Diagnostics);
 		FParsedShader MinifiedParsed = ParseShader(Minified);
 
 		// Expect true:
-		TestTrue(TEXT("MinifyShader: MainCS: contains MainCS"), ChunkPresent(MinifiedParsed, TEXT("MainCS")));
-		TestTrue(TEXT("MinifyShader: MainCS: contains FFoo"), ChunkPresent(MinifiedParsed, TEXT("FFoo")));
-		TestTrue(TEXT("MinifyShader: MainCS: contains FBar"), ChunkPresent(MinifiedParsed, TEXT("FBar")));
-		TestTrue(TEXT("MinifyShader: MainCS: contains Sum"), ChunkPresent(MinifiedParsed, TEXT("Sum")));
-		TestTrue(TEXT("MinifyShader: MainCS: contains FunA"), ChunkPresent(MinifiedParsed, TEXT("FunA")));
-		TestTrue(TEXT("MinifyShader: MainCS: contains FunB"), ChunkPresent(MinifiedParsed, TEXT("FunB")));
-		TestTrue(TEXT("MinifyShader: MainCS: contains GAnonymousStruct"), ChunkPresent(MinifiedParsed, TEXT("GAnonymousStruct")));
-		TestTrue(TEXT("MinifyShader: MainCS: contains GStructA"), ChunkPresent(MinifiedParsed, TEXT("GStructA")));
-		TestTrue(TEXT("MinifyShader: MainCS: contains GStructB"), ChunkPresent(MinifiedParsed, TEXT("GStructB")));
-		TestTrue(TEXT("MinifyShader: MainCS: contains GStructC"), ChunkPresent(MinifiedParsed, TEXT("GStructC")));
-		TestTrue(TEXT("MinifyShader: MainCS: contains GInitializedAnonymousStructA"), ChunkPresent(MinifiedParsed, TEXT("GInitializedAnonymousStructA")));
-		TestTrue(TEXT("MinifyShader: MainCS: contains GInitializedAnonymousStructB"), ChunkPresent(MinifiedParsed, TEXT("GInitializedAnonymousStructB")));
-		TestTrue(TEXT("MinifyShader: MainCS: contains OutputBufferType"), ChunkPresent(MinifiedParsed, TEXT("OutputBufferType")));
-		TestTrue(TEXT("MinifyShader: MainCS: contains OutputBuffer"), ChunkPresent(MinifiedParsed, TEXT("OutputBuffer")));
-		TestTrue(TEXT("MinifyShader: MainCS: contains FTypedefUsedStruct"), ChunkPresent(MinifiedParsed, TEXT("FTypedefUsedStruct")));
-		TestTrue(TEXT("MinifyShader: MainCS: contains FTypedefUsed"), ChunkPresent(MinifiedParsed, TEXT("FTypedefUsed")));
-		TestTrue(TEXT("MinifyShader: MainCS: contains FTypedefUsedChained"), ChunkPresent(MinifiedParsed, TEXT("FTypedefUsedChained")));
-		TestTrue(TEXT("MinifyShader: MainCS: contains TypedefUsedBuffer"), ChunkPresent(MinifiedParsed, TEXT("TypedefUsedBuffer")));
-		TestTrue(TEXT("MinifyShader: MainCS: contains struct TUnusedTemplate"), MinifiedParsed.Source.Contains(TEXT("struct TUsedTemplate")));
-		TestTrue(TEXT("MinifyShader: MainCS: contains TUsedTemplate<T> operator*"), MinifiedParsed.Source.Contains(TEXT("TUsedTemplate<T> operator*")));
-		TestTrue(TEXT("MinifyShader: MainCS: contains EEnumUsed"), MinifiedParsed.Source.Contains(TEXT("EEnumUsed")));
-		TestTrue(TEXT("MinifyShader: MainCS: contains ENUM_USED_PART_1"), MinifiedParsed.Source.Contains(TEXT("ENUM_USED_PART_1")));
-		TestTrue(TEXT("MinifyShader: MainCS: contains ENUM_USED_PART_2"), MinifiedParsed.Source.Contains(TEXT("ENUM_USED_PART_2")));
+		TestTrue(TEXT("MinifyShader: MainCS: contains MainCS"), ChunkPresent(MinifiedParsed, SHADER_SOURCE_LITERAL("MainCS")));
+		TestTrue(TEXT("MinifyShader: MainCS: contains FFoo"), ChunkPresent(MinifiedParsed, SHADER_SOURCE_LITERAL("FFoo")));
+		TestTrue(TEXT("MinifyShader: MainCS: contains FBar"), ChunkPresent(MinifiedParsed, SHADER_SOURCE_LITERAL("FBar")));
+		TestTrue(TEXT("MinifyShader: MainCS: contains Sum"), ChunkPresent(MinifiedParsed, SHADER_SOURCE_LITERAL("Sum")));
+		TestTrue(TEXT("MinifyShader: MainCS: contains FunA"), ChunkPresent(MinifiedParsed, SHADER_SOURCE_LITERAL("FunA")));
+		TestTrue(TEXT("MinifyShader: MainCS: contains FunB"), ChunkPresent(MinifiedParsed, SHADER_SOURCE_LITERAL("FunB")));
+		TestTrue(TEXT("MinifyShader: MainCS: contains GAnonymousStruct"), ChunkPresent(MinifiedParsed, SHADER_SOURCE_LITERAL("GAnonymousStruct")));
+		TestTrue(TEXT("MinifyShader: MainCS: contains GStructA"), ChunkPresent(MinifiedParsed, SHADER_SOURCE_LITERAL("GStructA")));
+		TestTrue(TEXT("MinifyShader: MainCS: contains GStructB"), ChunkPresent(MinifiedParsed, SHADER_SOURCE_LITERAL("GStructB")));
+		TestTrue(TEXT("MinifyShader: MainCS: contains GStructC"), ChunkPresent(MinifiedParsed, SHADER_SOURCE_LITERAL("GStructC")));
+		TestTrue(TEXT("MinifyShader: MainCS: contains GInitializedAnonymousStructA"), ChunkPresent(MinifiedParsed, SHADER_SOURCE_LITERAL("GInitializedAnonymousStructA")));
+		TestTrue(TEXT("MinifyShader: MainCS: contains GInitializedAnonymousStructB"), ChunkPresent(MinifiedParsed, SHADER_SOURCE_LITERAL("GInitializedAnonymousStructB")));
+		TestTrue(TEXT("MinifyShader: MainCS: contains OutputBufferType"), ChunkPresent(MinifiedParsed, SHADER_SOURCE_LITERAL("OutputBufferType")));
+		TestTrue(TEXT("MinifyShader: MainCS: contains OutputBuffer"), ChunkPresent(MinifiedParsed, SHADER_SOURCE_LITERAL("OutputBuffer")));
+		TestTrue(TEXT("MinifyShader: MainCS: contains FTypedefUsedStruct"), ChunkPresent(MinifiedParsed, SHADER_SOURCE_LITERAL("FTypedefUsedStruct")));
+		TestTrue(TEXT("MinifyShader: MainCS: contains FTypedefUsed"), ChunkPresent(MinifiedParsed, SHADER_SOURCE_LITERAL("FTypedefUsed")));
+		TestTrue(TEXT("MinifyShader: MainCS: contains FTypedefUsedChained"), ChunkPresent(MinifiedParsed, SHADER_SOURCE_LITERAL("FTypedefUsedChained")));
+		TestTrue(TEXT("MinifyShader: MainCS: contains TypedefUsedBuffer"), ChunkPresent(MinifiedParsed, SHADER_SOURCE_LITERAL("TypedefUsedBuffer")));
+		TestTrue(TEXT("MinifyShader: MainCS: contains struct TUnusedTemplate"), MinifiedParsed.Source.Contains(SHADER_SOURCE_LITERAL("struct TUsedTemplate")));
+		TestTrue(TEXT("MinifyShader: MainCS: contains TUsedTemplate<T> operator*"), MinifiedParsed.Source.Contains(SHADER_SOURCE_LITERAL("TUsedTemplate<T> operator*")));
+		TestTrue(TEXT("MinifyShader: MainCS: contains EEnumUsed"), MinifiedParsed.Source.Contains(SHADER_SOURCE_LITERAL("EEnumUsed")));
+		TestTrue(TEXT("MinifyShader: MainCS: contains ENUM_USED_PART_1"), MinifiedParsed.Source.Contains(SHADER_SOURCE_LITERAL("ENUM_USED_PART_1")));
+		TestTrue(TEXT("MinifyShader: MainCS: contains ENUM_USED_PART_2"), MinifiedParsed.Source.Contains(SHADER_SOURCE_LITERAL("ENUM_USED_PART_2")));
 
 		// Expect false:
-		TestFalse(TEXT("MinifyShader: MainCS: contains UnreferencedFunction"), ChunkPresent(MinifiedParsed, TEXT("UnreferencedFunction")));
-		TestFalse(TEXT("MinifyShader: MainCS: contains FUnreferencedStruct"), ChunkPresent(MinifiedParsed, TEXT("FUnreferencedStruct")));
-		TestFalse(TEXT("MinifyShader: MainCS: contains GUnreferencedParameter"), ChunkPresent(MinifiedParsed, TEXT("GUnreferencedParameter")));
-		TestFalse(TEXT("MinifyShader: MainCS: contains FTypedefUsedChainedUnused"), ChunkPresent(MinifiedParsed, TEXT("FTypedefUsedChainedUnused")));
-		TestFalse(TEXT("MinifyShader: MainCS: contains FTypedefUnusedStruct"), ChunkPresent(MinifiedParsed, TEXT("FTypedefUnusedStruct")));
-		TestFalse(TEXT("MinifyShader: MainCS: contains FTypedefUnused"), ChunkPresent(MinifiedParsed, TEXT("FTypedefUnused")));
-		TestFalse(TEXT("MinifyShader: MainCS: contains TypedefUnusedBuffer"), ChunkPresent(MinifiedParsed, TEXT("TypedefUnusedBuffer")));
-		TestFalse(TEXT("MinifyShader: MainCS: contains struct TUnusedTemplate"), MinifiedParsed.Source.Contains(TEXT("struct TUnusedTemplate")));
-		TestFalse(TEXT("MinifyShader: MainCS: contains TUnusedTemplate<T> operator%"), MinifiedParsed.Source.Contains(TEXT("TUnusedTemplate<T> operator%")));
-		TestFalse(TEXT("MinifyShader: MainCS: contains EEnumUnused"), MinifiedParsed.Source.Contains(TEXT("EEnumUnused")));
-		TestFalse(TEXT("MinifyShader: MainCS: contains ENUM_UNUSED_PART_1"), MinifiedParsed.Source.Contains(TEXT("ENUM_UNUSED_PART_1")));
-		TestFalse(TEXT("MinifyShader: MainCS: contains ENUM_UNUSED_PART_2"), MinifiedParsed.Source.Contains(TEXT("ENUM_UNUSED_PART_2")));
+		TestFalse(TEXT("MinifyShader: MainCS: contains UnreferencedFunction"), ChunkPresent(MinifiedParsed, SHADER_SOURCE_LITERAL("UnreferencedFunction")));
+		TestFalse(TEXT("MinifyShader: MainCS: contains FUnreferencedStruct"), ChunkPresent(MinifiedParsed, SHADER_SOURCE_LITERAL("FUnreferencedStruct")));
+		TestFalse(TEXT("MinifyShader: MainCS: contains GUnreferencedParameter"), ChunkPresent(MinifiedParsed, SHADER_SOURCE_LITERAL("GUnreferencedParameter")));
+		TestFalse(TEXT("MinifyShader: MainCS: contains FTypedefUsedChainedUnused"), ChunkPresent(MinifiedParsed, SHADER_SOURCE_LITERAL("FTypedefUsedChainedUnused")));
+		TestFalse(TEXT("MinifyShader: MainCS: contains FTypedefUnusedStruct"), ChunkPresent(MinifiedParsed, SHADER_SOURCE_LITERAL("FTypedefUnusedStruct")));
+		TestFalse(TEXT("MinifyShader: MainCS: contains FTypedefUnused"), ChunkPresent(MinifiedParsed, SHADER_SOURCE_LITERAL("FTypedefUnused")));
+		TestFalse(TEXT("MinifyShader: MainCS: contains TypedefUnusedBuffer"), ChunkPresent(MinifiedParsed, SHADER_SOURCE_LITERAL("TypedefUnusedBuffer")));
+		TestFalse(TEXT("MinifyShader: MainCS: contains struct TUnusedTemplate"), MinifiedParsed.Source.Contains(SHADER_SOURCE_LITERAL("struct TUnusedTemplate")));
+		TestFalse(TEXT("MinifyShader: MainCS: contains TUnusedTemplate<T> operator%"), MinifiedParsed.Source.Contains(SHADER_SOURCE_LITERAL("TUnusedTemplate<T> operator%")));
+		TestFalse(TEXT("MinifyShader: MainCS: contains EEnumUnused"), MinifiedParsed.Source.Contains(SHADER_SOURCE_LITERAL("EEnumUnused")));
+		TestFalse(TEXT("MinifyShader: MainCS: contains ENUM_UNUSED_PART_1"), MinifiedParsed.Source.Contains(SHADER_SOURCE_LITERAL("ENUM_UNUSED_PART_1")));
+		TestFalse(TEXT("MinifyShader: MainCS: contains ENUM_UNUSED_PART_2"), MinifiedParsed.Source.Contains(SHADER_SOURCE_LITERAL("ENUM_UNUSED_PART_2")));
 	}
 
 	int32 NumErrors = ExecutionInfo.GetErrorTotal();
diff --git a/Engine/Source/Developer/ShaderCompilerCommon/Public/ShaderCompilerCommon.h b/Engine/Source/Developer/ShaderCompilerCommon/Public/ShaderCompilerCommon.h
--- a/Engine/Source/Developer/ShaderCompilerCommon/Public/ShaderCompilerCommon.h
+++ b/Engine/Source/Developer/ShaderCompilerCommon/Public/ShaderCompilerCommon.h
@@ -11,6 +11,7 @@
 #include "Interfaces/IShaderFormat.h"
 
 class FShaderParameterParser;
+class FShaderSource;
 
 namespace UE::ShaderCompilerCommon
 {
@@ -98,9 +99,16 @@ namespace UE::ShaderCompilerCommon
 	* Any errors encountered during parsing or minification are added to OutErrors.
 	* InOutPreprocessedShaderSource is replaced with the rewritten code on success and is kept intact on failure.
 	*/
-	extern SHADERCOMPILERCOMMON_API bool RemoveDeadCode(FString& InOutPreprocessedShaderSource, const FString& EntryPoint, TArray<FShaderCompilerError>& OutErrors);
-	extern SHADERCOMPILERCOMMON_API bool RemoveDeadCode(FString& InOutPreprocessedShaderSource, const FString& EntryPoint, TConstArrayView<FStringView> RequiredSymbols, TArray<FShaderCompilerError>& OutErrors);
-	extern SHADERCOMPILERCOMMON_API bool RemoveDeadCode(FString& InOutPreprocessedShaderSource, TConstArrayView<FStringView> RequiredSymbols, TArray<FShaderCompilerError>& OutErrors);
+	extern SHADERCOMPILERCOMMON_API bool RemoveDeadCode(FShaderSource& InOutPreprocessedShaderSource, const FString& EntryPoint, TArray<FShaderCompilerError>& OutErrors);
+	extern SHADERCOMPILERCOMMON_API bool RemoveDeadCode(FShaderSource& InOutPreprocessedShaderSource, const FString& EntryPoint, TConstArrayView<FStringView> RequiredSymbols, TArray<FShaderCompilerError>& OutErrors);
+	extern SHADERCOMPILERCOMMON_API bool RemoveDeadCode(FShaderSource& InOutPreprocessedShaderSource, TConstArrayView<FStringView> RequiredSymbols, TArray<FShaderCompilerError>& OutErrors);
+	
+	UE_DEPRECATED(5.4, "Use overload of RemoveDeadCode accepting FShaderSource")
+	inline bool RemoveDeadCode(FString& InOutPreprocessedShaderSource, const FString& EntryPoint, TArray<FShaderCompilerError>& OutErrors) { return false; }
+	UE_DEPRECATED(5.4, "Use overload of RemoveDeadCode accepting FShaderSource")
+	inline bool RemoveDeadCode(FString& InOutPreprocessedShaderSource, const FString& EntryPoint, TConstArrayView<FStringView> RequiredSymbols, TArray<FShaderCompilerError>& OutErrors) { return false; }
+	UE_DEPRECATED(5.4, "Use overload of RemoveDeadCode accepting FShaderSource")
+	inline bool RemoveDeadCode(FString& InOutPreprocessedShaderSource, TConstArrayView<FStringView> RequiredSymbols, TArray<FShaderCompilerError>& OutErrors) { return false; }
 
 	struct FDebugShaderDataOptions
 	{
@@ -174,7 +182,13 @@ namespace UE::ShaderCompilerCommon
 	 * @param	PreprocessedSource The unmodified preprocessed source (used as input to the compilation)
 	 * @param	Options Options which can change behaviour of the debug dump; see above.
 	 */
-	extern SHADERCOMPILERCOMMON_API FString GetDebugShaderContents(const FShaderCompilerInput& Input, const FString& PreprocessedSource, const FDebugShaderDataOptions& Options = FDebugShaderDataOptions());
+	extern SHADERCOMPILERCOMMON_API FString GetDebugShaderContents(const FShaderCompilerInput& Input, FStringView PreprocessedSource, const FDebugShaderDataOptions& Options = FDebugShaderDataOptions());
+	
+	UE_DEPRECATED(5.4, "Use overload of GetDebugShaderContents accepting an FStringView")
+	inline FString GetDebugShaderContents(const FShaderCompilerInput& Input, const FString& PreprocessedSource, const FDebugShaderDataOptions& Options = FDebugShaderDataOptions())
+	{
+		return FString();
+	}
 
 	class FBaseShaderFormat : public IShaderFormat
 	{
@@ -329,12 +343,41 @@ extern SHADERCOMPILERCOMMON_API void AddUnboundShaderParameterError(
 	FShaderCompilerOutput& CompilerOutput);
 
 // Convert generated UniformBuffer code and references into something the shader compilers can use.
-extern SHADERCOMPILERCOMMON_API void CleanupUniformBufferCode(const FShaderCompilerEnvironment& Environment, FString& PreprocessedShaderSource);
+extern SHADERCOMPILERCOMMON_API void CleanupUniformBufferCode(const FShaderCompilerEnvironment& Environment, FShaderSource& PreprocessedShaderSource);
 
 UE_DEPRECATED(5.4, "RemoveUniformBuffersFromSource was renamed CleanupUniformBufferCode")
 inline void RemoveUniformBuffersFromSource(const FShaderCompilerEnvironment& Environment, FString& PreprocessedShaderSource) {}
 
-extern SHADERCOMPILERCOMMON_API const TCHAR* FindMatchingClosingBrace(const TCHAR* OpeningCharPtr);
+template <typename CharType>
+const CharType* FindMatchingBlock(const CharType* OpeningCharPtr, char OpenChar, char CloseChar)
+{
+	const CharType* SearchPtr = OpeningCharPtr;
+	int32 Depth = 0;
+
+	while (*SearchPtr)
+	{
+		if (*SearchPtr == OpenChar)
+		{
+			Depth++;
+		}
+		else if (*SearchPtr == CloseChar)
+		{
+			if (Depth == 0)
+			{
+				return SearchPtr;
+			}
+
+			Depth--;
+		}
+		SearchPtr++;
+	}
+
+	return nullptr;
+}
+
+template <typename CharType>
+const CharType* FindMatchingClosingBrace(const CharType* OpeningCharPtr) { return FindMatchingBlock<CharType>(OpeningCharPtr, '{', '}'); };
+
 extern SHADERCOMPILERCOMMON_API const TCHAR* ParseHLSLSymbolName(const TCHAR* SearchString, FString& SymboName);
 extern SHADERCOMPILERCOMMON_API void ParseHLSLTypeName(const TCHAR* SearchString, const TCHAR*& TypeNameStartPtr, const TCHAR*& TypeNameEndPtr);
 
diff --git a/Engine/Source/Developer/ShaderCompilerCommon/Public/ShaderConductorContext.h b/Engine/Source/Developer/ShaderCompilerCommon/Public/ShaderConductorContext.h
--- a/Engine/Source/Developer/ShaderCompilerCommon/Public/ShaderConductorContext.h
+++ b/Engine/Source/Developer/ShaderCompilerCommon/Public/ShaderConductorContext.h
@@ -201,6 +201,7 @@ namespace CrossCompiler
 		/** Loads the shader source and converts the input descriptor to a format suitable for ShaderConductor. If 'Definitions' is null, the previously loaded definitions are not modified. */
 		PRAGMA_DISABLE_DEPRECATION_WARNINGS		// FShaderCompilerDefinitions will be made internal in the future, marked deprecated until then
 		bool LoadSource(const FString& ShaderSource, const FString& Filename, const FString& EntryPoint, EShaderFrequency ShaderStage, const FShaderCompilerDefinitions* Definitions = nullptr, const TArray<FString>* ExtraDxcArgs = nullptr);
+		bool LoadSource(FStringView ShaderSource, const FString& Filename, const FString& EntryPoint, EShaderFrequency ShaderStage, const FShaderCompilerDefinitions* Definitions = nullptr, const TArray<FString>* ExtraDxcArgs = nullptr);
 		bool LoadSource(const ANSICHAR* ShaderSource, const ANSICHAR* Filename, const ANSICHAR* EntryPoint, EShaderFrequency ShaderStage, const FShaderCompilerDefinitions* Definitions = nullptr, const TArray<FString>* ExtraDxcArgs = nullptr);
 		PRAGMA_ENABLE_DEPRECATION_WARNINGS
 
diff --git a/Engine/Source/Developer/ShaderCompilerCommon/Public/ShaderMinifier.h b/Engine/Source/Developer/ShaderCompilerCommon/Public/ShaderMinifier.h
--- a/Engine/Source/Developer/ShaderCompilerCommon/Public/ShaderMinifier.h
+++ b/Engine/Source/Developer/ShaderCompilerCommon/Public/ShaderMinifier.h
@@ -6,6 +6,7 @@
 
 #include "Containers/UnrealString.h"
 #include "Misc/EnumClassFlags.h"
+#include "ShaderPreprocessTypes.h"
 
 namespace UE::ShaderMinifier
 {
@@ -45,7 +46,7 @@ struct FDiagnostics
 
 struct FMinifiedShader
 {
-	FString Code;
+	FShaderSource::FStringType Code;
 	FDiagnostics Diagnostics;
 
 	bool Success() const
@@ -54,8 +55,14 @@ struct FMinifiedShader
 	}
 };
 
-extern SHADERCOMPILERCOMMON_API FMinifiedShader Minify(const FStringView PreprocessedShader, const FStringView EntryPoint, EMinifyShaderFlags Flags = EMinifyShaderFlags::None);
-extern SHADERCOMPILERCOMMON_API FMinifiedShader Minify(const FStringView PreprocessedShader, TConstArrayView<FStringView> RequiredSymbols, EMinifyShaderFlags Flags = EMinifyShaderFlags::None);
+extern SHADERCOMPILERCOMMON_API FMinifiedShader Minify(const FShaderSource& PreprocessedShader, const FShaderSource::FViewType EntryPoint, EMinifyShaderFlags Flags = EMinifyShaderFlags::None);
+extern SHADERCOMPILERCOMMON_API FMinifiedShader Minify(const FShaderSource& PreprocessedShader, TConstArrayView<FShaderSource::FViewType> RequiredSymbols, EMinifyShaderFlags Flags = EMinifyShaderFlags::None);
+
+UE_DEPRECATED(5.4, "Use Minify overloads accepting FShaderSource")
+inline FMinifiedShader Minify(const FStringView PreprocessedShader, const FStringView EntryPoint, EMinifyShaderFlags Flags = EMinifyShaderFlags::None) { return FMinifiedShader(); }
+
+UE_DEPRECATED(5.4, "Use Minify overloads accepting FShaderSource")
+inline FMinifiedShader Minify(const FStringView PreprocessedShader, TConstArrayView<FStringView> RequiredSymbols, EMinifyShaderFlags Flags = EMinifyShaderFlags::None) { return FMinifiedShader(); }
 
 } // UE::ShaderMinifier
 
diff --git a/Engine/Source/Developer/ShaderFormatOpenGL/Private/OpenGLShaderCompiler.cpp b/Engine/Source/Developer/ShaderFormatOpenGL/Private/OpenGLShaderCompiler.cpp
--- a/Engine/Source/Developer/ShaderFormatOpenGL/Private/OpenGLShaderCompiler.cpp
+++ b/Engine/Source/Developer/ShaderFormatOpenGL/Private/OpenGLShaderCompiler.cpp
@@ -3265,10 +3265,10 @@ bool ShouldUseDXC(FShaderCompilerFlags Flags)
  * @param WorkingDirectory - (unused, part of IShaderFormat API)
  * @param Version - Target GLSL version for this compilation
  */
-void CompileOpenGLShader(const FShaderCompilerInput& Input, const FString& InPreprocessedSource, FShaderCompilerOutput& Output, const FString& WorkingDirectory, GLSLVersion Version)
+void CompileOpenGLShader(const FShaderCompilerInput& Input, const FShaderPreprocessOutput& InPreprocessOutput, FShaderCompilerOutput& Output, const FString& WorkingDirectory, GLSLVersion Version)
 {
 	FString EntryPointName = Input.EntryPointName;
-	FString PreprocessedSource = InPreprocessedSource;
+	FString PreprocessedSource(InPreprocessOutput.GetSourceViewWide());
 
 	FShaderParameterParser::FPlatformConfiguration PlatformConfiguration;
 	FShaderParameterParser ShaderParameterParser(PlatformConfiguration);
@@ -3419,7 +3419,7 @@ void FOpenGLFrontend::CompileShader(const FShaderCompilerInput& Input, FShaderCo
 {
 	FShaderPreprocessOutput PreprocessOutput;
 	PreprocessShader(PreprocessOutput, Input, Input.Environment);
-	CompileOpenGLShader(Input, PreprocessOutput.GetSource(), Output, WorkingDirectory, Version);
+	CompileOpenGLShader(Input, PreprocessOutput, Output, WorkingDirectory, Version);
 }
 
 static void FillDeviceCapsOfflineCompilationInternal(struct FDeviceCapabilities& Capabilities, const GLSLVersion ShaderVersion)
diff --git a/Engine/Source/Developer/ShaderFormatOpenGL/Private/ShaderFormatOpenGL.cpp b/Engine/Source/Developer/ShaderFormatOpenGL/Private/ShaderFormatOpenGL.cpp
--- a/Engine/Source/Developer/ShaderFormatOpenGL/Private/ShaderFormatOpenGL.cpp
+++ b/Engine/Source/Developer/ShaderFormatOpenGL/Private/ShaderFormatOpenGL.cpp
@@ -20,7 +20,7 @@ extern bool ShouldUseDXC(FShaderCompilerFlags Flags);
 
 extern void CompileOpenGLShader(
 	const FShaderCompilerInput& Input,
-	const FString& InPreprocessedSource,
+	const FShaderPreprocessOutput& InPreprocessOutput,
 	FShaderCompilerOutput& Output,
 	const FString& WorkingDirectory,
 	GLSLVersion Version);
@@ -109,7 +109,7 @@ public:
 	virtual void CompilePreprocessedShader(const FShaderCompilerInput& Input, const FShaderPreprocessOutput& PreprocessOutput, FShaderCompilerOutput& Output, const FString& WorkingDirectory) const override
 	{
 		CheckFormat(Input.ShaderFormat);
-		CompileOpenGLShader(Input, PreprocessOutput.GetSource(), Output, WorkingDirectory, TranslateFormatNameToEnum(Input.ShaderFormat));		
+		CompileOpenGLShader(Input, PreprocessOutput, Output, WorkingDirectory, TranslateFormatNameToEnum(Input.ShaderFormat));		
 	}
 
 	virtual const TCHAR* GetPlatformIncludeDirectory() const override
diff --git a/Engine/Source/Developer/ShaderFormatVectorVM/Private/VectorVMShaderCompiler.cpp b/Engine/Source/Developer/ShaderFormatVectorVM/Private/VectorVMShaderCompiler.cpp
--- a/Engine/Source/Developer/ShaderFormatVectorVM/Private/VectorVMShaderCompiler.cpp
+++ b/Engine/Source/Developer/ShaderFormatVectorVM/Private/VectorVMShaderCompiler.cpp
@@ -51,7 +51,8 @@ bool CompileVectorVMShader(
 	}
 
 	FVectorVMCodeBackend VVMBackEnd(CCFlags, HlslCompilerTarget, VMCompilationOutput);
-	FVectorVMLanguageSpec VVMLanguageSpec; 
+	FVectorVMLanguageSpec VVMLanguageSpec;
+	FAnsiStringView Source = PreprocessOutput.GetSourceViewAnsi();
 
 	bool bResult = false;
 	{
@@ -62,7 +63,7 @@ bool CompileVectorVMShader(
 			SCOPE_CYCLE_COUNTER(STAT_VectorVM_Compiler_CompileShader_CrossCompilerContextRun);
 
 			bResult = CrossCompilerContext.Run(
-				TCHAR_TO_ANSI(*PreprocessOutput.GetSource()),
+				Source.GetData(),
 				TCHAR_TO_ANSI(*Input.EntryPointName),
 				&VVMBackEnd,
 				&ShaderSource,
@@ -82,16 +83,15 @@ bool CompileVectorVMShader(
 
 		VMCompilationOutput.Errors = Converted.GetData();
 
-		TArray<FString> OutputByLines;
-		PreprocessOutput.GetSource().ParseIntoArrayLines(OutputByLines, false);
-
 		UE_LOG(LogVectorVMShaderCompiler, Warning, TEXT("Warnings while processing %s"), *Input.DebugGroupName);
 
-		FString OutputHlsl;
-		for (int32 i = 0; i < OutputByLines.Num(); i++)
-		{
-			UE_LOG(LogVectorVMShaderCompiler, Display, TEXT("/*%d*/%s"), i, *OutputByLines[i]);
-		}
+		PreprocessOutput.ForEachLine
+		(
+			[](FAnsiStringView Line, int32 LineIndex)
+			{
+				UE_LOG(LogVectorVMShaderCompiler, Display, TEXT("/*%d*/%.*hs"), LineIndex, Line.Len(), Line.GetData());
+			}
+		);
 	}
 	else
 	{
diff --git a/Engine/Source/Developer/ShaderPreprocessor/Private/ShaderPreprocessor.cpp b/Engine/Source/Developer/ShaderPreprocessor/Private/ShaderPreprocessor.cpp
--- a/Engine/Source/Developer/ShaderPreprocessor/Private/ShaderPreprocessor.cpp
+++ b/Engine/Source/Developer/ShaderPreprocessor/Private/ShaderPreprocessor.cpp
@@ -834,8 +834,8 @@ bool PreprocessShader(
 		// Append ShaderPrint generated code at the end of the shader if necessary
 		Context.ShaderPrintGenerate(OutPreprocessedAnsi, &Output.EditDiagnosticDatas());
 
-		// "preprocessor_file_size" includes null terminator, so subtract one for Append call -- passing size saves an expensive strlen in Append
-		Output.EditSource().Append(OutPreprocessedAnsi, preprocessor_file_size(OutPreprocessedAnsi) - 1);
+		// "preprocessor_file_size" includes null terminator, so subtract one when initializing the FShaderSource (which automatically null terminates)
+		Output.EditSource().Set({ OutPreprocessedAnsi, preprocessor_file_size(OutPreprocessedAnsi) - 1 });
 	}
 
 	if (!HasError && !Context.HasIncludedMandatoryHeaders())
@@ -866,7 +866,7 @@ bool PreprocessShader(
 	const FShaderCompilerEnvironment& Environment = ShaderInput.Environment;
 	bool bSucceeded = PreprocessShader(Output, ShaderInput, Environment, AdditionalDefines);
 
-	OutPreprocessedShader = MoveTemp(Output.EditSource());
+	OutPreprocessedShader = FString(Output.GetSourceViewWide());
 
 	Output.MoveDirectives(ShaderOutput.PragmaDirectives);
 	for (FShaderCompilerError& Error : Output.EditErrors())
diff --git a/Engine/Source/Developer/VulkanShaderFormat/Private/VulkanShaderCompiler.cpp b/Engine/Source/Developer/VulkanShaderFormat/Private/VulkanShaderCompiler.cpp
--- a/Engine/Source/Developer/VulkanShaderFormat/Private/VulkanShaderCompiler.cpp
+++ b/Engine/Source/Developer/VulkanShaderFormat/Private/VulkanShaderCompiler.cpp
@@ -1850,7 +1850,7 @@ static FString ParseEntrypointDecl(FStringView PreprocessedShader, FStringView E
 
 uint8 ParseWaveSize(
 	const FVulkanShaderCompilerInternalState& InternalState,
-	const FString& PreprocessedShader
+	FStringView PreprocessedShader
 	)
 {
 	uint8 WaveSize = 0;
@@ -1902,7 +1902,7 @@ uint8 ParseWaveSize(
 
 static bool CompileWithShaderConductor(
 	const FVulkanShaderCompilerInternalState& InternalState,
-	const FString&			PreprocessedShader,
+	FStringView PreprocessedShader,
 	VulkanShaderCompilerSerializedOutput& SerializedOutput,
 	FShaderCompilerOutput&	Output
 )
@@ -2282,15 +2282,16 @@ static bool CompileShaderGroup(
 		const bool bIsClosestHit = (HitGroupShaderType == FVulkanShaderCompilerInternalState::EHitGroupShaderType::ClosestHit);
 		FShaderCompilerOutput& PartialOutput = bIsClosestHit ? MergedOutput : TempOutput;
 
-		FString PartialPreprocessedShaderSource = OriginalPreprocessedShaderSource;
+		FStringView OrigSourceView(OriginalPreprocessedShaderSource);
+		FShaderSource PartialPreprocessedShaderSource(OrigSourceView);
 		UE::ShaderCompilerCommon::RemoveDeadCode(PartialPreprocessedShaderSource, InternalState.GetEntryPointName(), PartialOutput.Errors);
 
 		if (InternalState.bDebugDump)
 		{
-			DumpDebugShaderText(InternalState.Input, PartialPreprocessedShaderSource, *FString::Printf(TEXT("%s.hlsl"), PartialFileExtension));
+			DumpDebugShaderText(InternalState.Input, PartialPreprocessedShaderSource.GetView().GetData(), *FString::Printf(TEXT("%s.hlsl"), PartialFileExtension));
 		}
 
-		const bool bPartialSuccess = CompileWithShaderConductor(InternalState, PartialPreprocessedShaderSource, PartialSerializedOutput, PartialOutput);
+		const bool bPartialSuccess = CompileWithShaderConductor(InternalState, PartialPreprocessedShaderSource.GetView(), PartialSerializedOutput, PartialOutput);
 
 		if (!bIsClosestHit)
 		{
@@ -2389,12 +2390,12 @@ struct FPS5ShaderParameterParserPlatformConfiguration : public FShaderParameterP
 	}
 };
 
-void CompileVulkanShader(const FShaderCompilerInput& Input, const FString& InPreprocessedSource, FShaderCompilerOutput& Output, const class FString& WorkingDirectory)
+void CompileVulkanShader(const FShaderCompilerInput& Input, const FShaderPreprocessOutput& InPreprocessOutput, FShaderCompilerOutput& Output, const class FString& WorkingDirectory)
 {
 	check(IsVulkanShaderFormat(Input.ShaderFormat));
 
 	FString EntryPointName = Input.EntryPointName;
-	FString PreprocessedSource = InPreprocessedSource;
+	FString PreprocessedSource(InPreprocessOutput.GetSourceViewWide());
 
 	FPS5ShaderParameterParserPlatformConfiguration PlatformConfiguration(Input);
 	FShaderParameterParser ShaderParameterParser(PlatformConfiguration);
diff --git a/Engine/Source/Developer/VulkanShaderFormat/Private/VulkanShaderFormat.cpp b/Engine/Source/Developer/VulkanShaderFormat/Private/VulkanShaderFormat.cpp
--- a/Engine/Source/Developer/VulkanShaderFormat/Private/VulkanShaderFormat.cpp
+++ b/Engine/Source/Developer/VulkanShaderFormat/Private/VulkanShaderFormat.cpp
@@ -21,7 +21,7 @@ extern void ModifyVulkanCompilerInput(FShaderCompilerInput& Input);
 
 extern void CompileVulkanShader(
 	const FShaderCompilerInput& Input,
-	const FString& InPreprocessedSource,
+	const FShaderPreprocessOutput& InPreprocessOutput,
 	FShaderCompilerOutput& Output,
 	const FString& WorkingDirectory);
 
@@ -95,7 +95,7 @@ public:
 
 	virtual void CompilePreprocessedShader(const FShaderCompilerInput& Input, const FShaderPreprocessOutput& PreprocessOutput, FShaderCompilerOutput& Output,const FString& WorkingDirectory) const override
 	{
-		CompileVulkanShader(Input, PreprocessOutput.GetSource(), Output, WorkingDirectory);
+		CompileVulkanShader(Input, PreprocessOutput, Output, WorkingDirectory);
 	}
 
 	virtual void OutputDebugData(const FShaderCompilerInput& Input, const FShaderPreprocessOutput& PreprocessOutput, const FShaderCompilerOutput& Output) const override
diff --git a/Engine/Source/Developer/Windows/ShaderFormatD3D/Private/D3DShaderCompiler.cpp b/Engine/Source/Developer/Windows/ShaderFormatD3D/Private/D3DShaderCompiler.cpp
--- a/Engine/Source/Developer/Windows/ShaderFormatD3D/Private/D3DShaderCompiler.cpp
+++ b/Engine/Source/Developer/Windows/ShaderFormatD3D/Private/D3DShaderCompiler.cpp
@@ -1280,7 +1280,7 @@ struct FD3DShaderParameterParserPlatformConfiguration : public FShaderParameterP
 	}
 };
 
-void CompileD3DShader(const FShaderCompilerInput& Input, const FString& InPreprocessedSource, FShaderCompilerOutput& Output, const FString& WorkingDirectory, ED3DShaderModel ShaderModel)
+void CompileD3DShader(const FShaderCompilerInput& Input, const FShaderPreprocessOutput& InPreprocessOutput, FShaderCompilerOutput& Output, const FString& WorkingDirectory, ED3DShaderModel ShaderModel)
 {
 	TRACE_CPUPROFILER_EVENT_SCOPE(CompileD3DShader);
 
@@ -1293,7 +1293,7 @@ void CompileD3DShader(const FShaderCompilerInput& Input, const FString& InPrepro
 	}
 
 	FString EntryPointName = Input.EntryPointName;
-	FString PreprocessedSource = InPreprocessedSource;
+	FString PreprocessedSource(InPreprocessOutput.GetSourceViewWide());
 
 	FD3DShaderParameterParserPlatformConfiguration PlatformConfiguration;
 	FShaderParameterParser ShaderParameterParser(PlatformConfiguration);
diff --git a/Engine/Source/Developer/Windows/ShaderFormatD3D/Private/ShaderFormatD3D.cpp b/Engine/Source/Developer/Windows/ShaderFormatD3D/Private/ShaderFormatD3D.cpp
--- a/Engine/Source/Developer/Windows/ShaderFormatD3D/Private/ShaderFormatD3D.cpp
+++ b/Engine/Source/Developer/Windows/ShaderFormatD3D/Private/ShaderFormatD3D.cpp
@@ -177,7 +177,7 @@ public:
 		FShaderCompilerOutput& Output,
 		const FString& WorkingDirectory) const
 	{
-		CompileD3DShader(Input, PreprocessOutput.GetSource(), Output, WorkingDirectory, DetermineShaderModel(Input));
+		CompileD3DShader(Input, PreprocessOutput, Output, WorkingDirectory, DetermineShaderModel(Input));
 
 		Output.ShaderDiagnosticDatas = PreprocessOutput.GetDiagnosticDatas();
 	}
diff --git a/Engine/Source/Developer/Windows/ShaderFormatD3D/Private/ShaderFormatD3D.h b/Engine/Source/Developer/Windows/ShaderFormatD3D/Private/ShaderFormatD3D.h
--- a/Engine/Source/Developer/Windows/ShaderFormatD3D/Private/ShaderFormatD3D.h
+++ b/Engine/Source/Developer/Windows/ShaderFormatD3D/Private/ShaderFormatD3D.h
@@ -27,7 +27,7 @@ bool PreprocessD3DShader(
 
 void CompileD3DShader(
 	const FShaderCompilerInput& Input,
-	const FString& InPreprocessedSource,
+	const FShaderPreprocessOutput& InPreprocessOutput,
 	FShaderCompilerOutput& Output,
 	const FString& WorkingDirectory,
 	ED3DShaderModel ShaderModel);
diff --git a/Engine/Source/Runtime/Engine/Private/Materials/MaterialShader.cpp b/Engine/Source/Runtime/Engine/Private/Materials/MaterialShader.cpp
--- a/Engine/Source/Runtime/Engine/Private/Materials/MaterialShader.cpp
+++ b/Engine/Source/Runtime/Engine/Private/Materials/MaterialShader.cpp
@@ -2356,7 +2356,7 @@ FShader* FMaterialShaderMap::ProcessCompilationResultsForSingleJob(FShaderCompil
 		const int32 Index = Algo::LowerBoundBy(GetMutableContent()->ShaderProcessedSource, Key, [](const FMaterialProcessedSource& Value) { return Value.Name; });
 		if (Index >= GetMutableContent()->ShaderProcessedSource.Num() || GetMutableContent()->ShaderProcessedSource[Index].Name != Key)
 		{
-			GetMutableContent()->ShaderProcessedSource.EmplaceAt(Index, Key, *CurrentJob.GetFinalSource());
+			GetMutableContent()->ShaderProcessedSource.EmplaceAt(Index, Key, CurrentJob.GetFinalSourceView().GetData());
 		}
 	}
 
diff --git a/Engine/Source/Runtime/RenderCore/Private/ShaderCore.cpp b/Engine/Source/Runtime/RenderCore/Private/ShaderCore.cpp
--- a/Engine/Source/Runtime/RenderCore/Private/ShaderCore.cpp
+++ b/Engine/Source/Runtime/RenderCore/Private/ShaderCore.cpp
@@ -3683,10 +3683,10 @@ FShaderCommonCompileJob::FInputHash FShaderCompileJob::GetInputHash()
 
 		PreprocessOutput.VisitDirectivesWithPrefix(TEXT("VERSION"), HashVersion);
 		// const_cast due to serialization API requiring non-const. better than not having const correctness in the API.
-		Hasher << const_cast<FString&>(PreprocessOutput.GetSource());
+		Hasher << PreprocessOutput.EditSource();
 		if (SecondaryPreprocessOutput.IsValid())
 		{
-			Hasher << const_cast<FString&>(SecondaryPreprocessOutput->GetSource());
+			Hasher << SecondaryPreprocessOutput->EditSource();
 		}
 
 		if (Input.RootParametersStructure)
@@ -3974,7 +3974,7 @@ void FShaderCompileJob::SerializeWorkerInput(FArchive& Ar)
 	}
 }
 
-const FString& FShaderCompileJob::GetFinalSource() const
+FStringView FShaderCompileJob::GetFinalSourceView() const
 {
 	 // any modifications to the source done as part of the compile step will be written to the "ModifiedShaderSource" field
 	// always return empty string if source extraction was not requested; this will prevent bloat of material DDC data in the case where debug info is enabled 
@@ -3983,12 +3983,11 @@ const FString& FShaderCompileJob::GetFinalSource() const
 	{
 		// if there are no such modifications, return the "unstripped" version of the source code (with comments & line directives maintained),
 		// otherwise return whatever the final modified source is as input to the compiler by the backend.
-		return Output.ModifiedShaderSource.IsEmpty() ? PreprocessOutput.GetUnstrippedSource() : Output.ModifiedShaderSource;
+		return Output.ModifiedShaderSource.IsEmpty() ? PreprocessOutput.GetUnstrippedSourceView() : FStringView(Output.ModifiedShaderSource);
 	}
 	else
 	{
-		static FString Empty;
-		return Empty;
+		return FStringView();
 	}
 }
 
diff --git a/Engine/Source/Runtime/RenderCore/Private/ShaderPreprocessTypes.cpp b/Engine/Source/Runtime/RenderCore/Private/ShaderPreprocessTypes.cpp
--- a/Engine/Source/Runtime/RenderCore/Private/ShaderPreprocessTypes.cpp
+++ b/Engine/Source/Runtime/RenderCore/Private/ShaderPreprocessTypes.cpp
@@ -3,10 +3,12 @@
 #include "ShaderPreprocessTypes.h"
 
 #include "ShaderCompilerCore.h"
+#include "Containers/AnsiString.h"
+#include "Templates/UnrealTemplate.h"
 
-const TCHAR* FilenameSentinel = TEXT("__UE_FILENAME_SENTINEL__");
-static const int FilenameSentinelLen = FCString::Strlen(FilenameSentinel);
-static const FString LineDirectiveSentinel = FString::Printf(TEXT("#line 1 \"%s\"\n"), FilenameSentinel);
+const FShaderSource::CharType* FilenameSentinel = SHADER_SOURCE_LITERAL("__UE_FILENAME_SENTINEL__");
+static const int FilenameSentinelLen = FShaderSource::FCStringType::Strlen(FilenameSentinel);
+static const FShaderSource::FStringType LineDirectiveSentinel = FShaderSource::FStringType::Printf(SHADER_SOURCE_LITERAL("#line 1 \"%s\"\n"), FilenameSentinel);
 
 void FShaderDiagnosticRemapper::Remap(FShaderCompilerError& Diagnostic) const
 {
@@ -131,33 +133,32 @@ FShaderDiagnosticRemapper::FRemapData FShaderDiagnosticRemapper::GetRemapData(in
 	return FRemapData{ FoundBlock.OriginalPath, FoundBlock.OriginalLineNum + StrippedLineOffsets[StrippedLineNum - 1] };
 }
 
-inline bool IsEndOfTheLine(TCHAR C)
+inline bool IsEndOfTheLine(FShaderSource::CharType C)
 {
-	return C == TEXT('\r') || C == TEXT('\n');
+	return C == '\r' || C == '\n';
 }
 
-inline bool StripNeedsHandling(TCHAR C)
+inline bool StripNeedsHandling(FShaderSource::CharType C)
 {
-	return IsEndOfTheLine(C) || C == TEXT('/') || C == 0 || C == TEXT('#');
+	return IsEndOfTheLine(C) || C == '/' || C == 0 || C == '#';
 }
 
-inline void SkipNewLine(const TCHAR*& Current, const TCHAR* End)
+inline void SkipNewLine(const FShaderSource::CharType*& Current, const FShaderSource::CharType* End)
 {
-	TCHAR First = Current < End ? Current[0] : 0;
-	TCHAR Second = Current + 1 < End ? Current[1] : 0;
-	Current += ((First + Second) == TEXT('\r') + TEXT('\n')) ? 2 : 1;
+	FShaderSource::CharType First = Current < End ? Current[0] : 0;
+	FShaderSource::CharType Second = Current + 1 < End ? Current[1] : 0;
+	Current += ((First + Second) == '\r' + '\n') ? 2 : 1;
 }
 
 void FShaderPreprocessOutput::StripCode()
 {
 	// Reserve worst case slack (i.e. assuming there is nothing to strip) to avoid reallocation
-	FString PreprocessedSourceStripped(LineDirectiveSentinel, PreprocessedSource.Len() + 1);
-	PreprocessedSourceStripped.GetCharArray().AddUninitialized(PreprocessedSourceStripped.GetCharArray().GetSlack());
-	TCHAR* OutStrippedData = PreprocessedSourceStripped.GetCharArray().GetData();
-	TCHAR* OutStripped = OutStrippedData + LineDirectiveSentinel.Len();
+	FShaderSource PreprocessedSourceStripped(LineDirectiveSentinel.GetCharArray().GetData(), PreprocessedSource.Len());
+	FShaderSource::CharType* OutStrippedData = PreprocessedSourceStripped.GetData();
+	FShaderSource::CharType* OutStripped = OutStrippedData + LineDirectiveSentinel.Len();
 
-	const TCHAR* Begin = PreprocessedSource.GetCharArray().GetData(), *Current = Begin;
-	const TCHAR* End = Current + PreprocessedSource.Len();
+	const FShaderSource::CharType* Begin = PreprocessedSource.GetData(), * Current = Begin;
+	const FShaderSource::CharType* End = Current + PreprocessedSource.Len();
 	int32 CurrentBlockUnstrippedLineOffset = 0;
 	int32 CurrentStrippedLineNum = 1;
 	while (Current < End)
@@ -178,25 +179,25 @@ void FShaderPreprocessOutput::StripCode()
 				Remapper.AddStrippedLine(CurrentStrippedLineNum, CurrentBlockUnstrippedLineOffset);
 
 				// normalize line endings
-				*OutStripped++ = TEXT('\n');
+				*OutStripped++ = '\n';
 				CurrentStrippedLineNum++;
 			}
 			CurrentBlockUnstrippedLineOffset++;
 			SkipNewLine(Current, End);
 		}
-		else if (Current[0] == TEXT('/'))
+		else if (Current[0] == '/')
 		{
-			if (Current[1] == TEXT('/'))
+			if (Current[1] == '/')
 			{
 				while (!IsEndOfTheLine(*Current) && Current < End)
 				{
 					++Current;
 				}
 			}
-			else if (Current[1] == TEXT('*'))
+			else if (Current[1] == '*')
 			{
 				Current += 2;
-				while (Current < End && !(Current[0] == TEXT('*') && Current[1] == TEXT('/')))
+				while (Current < End && !(Current[0] == '*' && Current[1] == '/'))
 				{
 					++Current;
 				}
@@ -207,13 +208,13 @@ void FShaderPreprocessOutput::StripCode()
 				*OutStripped++ = *Current++;
 			}
 		}
-		else if (Current[0] == TEXT('#'))
+		else if (Current[0] == '#')
 		{
 			if (Current + 4 < End
-				&& Current[1] == TEXT('l')
-				&& Current[2] == TEXT('i')
-				&& Current[3] == TEXT('n')
-				&& Current[4] == TEXT('e'))
+				&& Current[1] == 'l'
+				&& Current[2] == 'i'
+				&& Current[3] == 'n'
+				&& Current[4] == 'e')
 			{
 				Current += 4;
 
@@ -225,7 +226,7 @@ void FShaderPreprocessOutput::StripCode()
 				int32 DirectiveLineNumber = 0;
 				while (FChar::IsDigit(*Current) && Current < End)
 				{
-					DirectiveLineNumber = 10 * DirectiveLineNumber + (*Current - TEXT('0'));
+					DirectiveLineNumber = 10 * DirectiveLineNumber + (*Current - '0');
 					++Current;
 				}
 
@@ -239,7 +240,7 @@ void FShaderPreprocessOutput::StripCode()
 				{
 					++Current;
 
-					const TCHAR* DirectiveFileNameStart = Current;
+					const FShaderSource::CharType* DirectiveFileNameStart = Current;
 					int32 DirectiveFileNameLen = 0;
 
 					// count chars in directive filename
@@ -275,10 +276,9 @@ void FShaderPreprocessOutput::StripCode()
 			}
 		}
 	}
-	// Null terminate after stripped copy
-	*OutStripped++ = 0;
-	check(OutStripped <= OutStrippedData + PreprocessedSourceStripped.GetCharArray().Num());
-	PreprocessedSourceStripped.GetCharArray().SetNum((int32)(OutStripped - OutStrippedData));
+	check(OutStripped <= OutStrippedData + PreprocessedSourceStripped.Len());
+	// ShrinkToLen null terminates for us by virtue of adding zero'd SIMD padding
+	PreprocessedSourceStripped.ShrinkToLen((int32)(OutStripped - OutStrippedData));
 
 	OriginalPreprocessedSource = MoveTemp(PreprocessedSource);
 	PreprocessedSource = MoveTemp(PreprocessedSourceStripped);
diff --git a/Engine/Source/Runtime/RenderCore/Private/ShaderSource.cpp b/Engine/Source/Runtime/RenderCore/Private/ShaderSource.cpp
new file mode 100644
--- /dev/null
+++ b/Engine/Source/Runtime/RenderCore/Private/ShaderSource.cpp
@@ -0,0 +1,41 @@
+// Copyright Epic Games, Inc. All Rights Reserved.
+
+#include "ShaderSource.h"
+#include "Containers/StringConv.h"
+#include "HAL/UnrealMemory.h"
+
+FShaderSource::FShaderSource(FShaderSource::FViewType InSrc, int32 AdditionalSlack)
+{
+	Set(InSrc, AdditionalSlack);
+}
+
+void FShaderSource::Set(FShaderSource::FViewType InSrc, int32 AdditionalSlack)
+{
+	SetLen(InSrc.Len() + AdditionalSlack);
+	FShaderSource::CharType* SourceData = Source.GetData();
+	FMemory::Memcpy(SourceData, InSrc.GetData(), sizeof(FShaderSource::CharType) * InSrc.Len());
+}
+
+void FShaderSource::Set(FAnsiStringView InSrc)
+{
+	TStringConvert<ANSICHAR, TCHAR> Convert;
+	SetLen(InSrc.Len());
+	Convert.Convert(Source.GetData(), Source.Num(), InSrc.GetData(), InSrc.Len());
+}
+
+FShaderSource& FShaderSource::operator=(FShaderSource::FStringType&& InSrc)
+{
+	int32 InInitialLen = InSrc.Len();
+	// input char array already has a null terminator appended, so can add one less padding char
+	TArray<CharType>& InSrcArray = InSrc.GetCharArray();
+	InSrcArray.AddZeroed(ShaderSourceSimdPadding - 1);
+	Source = MoveTemp(InSrcArray);
+	return *this;
+}
+
+FArchive& operator<<(FArchive& Ar, FShaderSource& Src)
+{
+	Ar << Src.Source;
+	return Ar;
+}
+
diff --git a/Engine/Source/Runtime/RenderCore/Public/ShaderCompilerCore.h b/Engine/Source/Runtime/RenderCore/Public/ShaderCompilerCore.h
--- a/Engine/Source/Runtime/RenderCore/Public/ShaderCompilerCore.h
+++ b/Engine/Source/Runtime/RenderCore/Public/ShaderCompilerCore.h
@@ -30,8 +30,8 @@ class FShaderCompileJob;
 class FShaderPipelineCompileJob;
 typedef TSharedPtr<TArray<ANSICHAR>, ESPMode::ThreadSafe> FShaderSharedAnsiStringPtr;
 
-// this is for the protocol, not the data, bump if FShaderCompilerInput or ProcessInputFromArchive changes.
-inline const int32 ShaderCompileWorkerInputVersion = 25;
+// this is for the protocol, not the data, bump if FShaderCompilerInput/FShaderPreprocessOutput serialization, SerializeWorkerInput or ProcessInputFromArchive changes.
+inline const int32 ShaderCompileWorkerInputVersion = 26;
 // this is for the protocol, not the data, bump if FShaderCompilerOutput or WriteToOutputArchive changes.
 inline const int32 ShaderCompileWorkerOutputVersion = 19;
 // this is for the protocol, not the data.
diff --git a/Engine/Source/Runtime/RenderCore/Public/ShaderCompilerJobTypes.h b/Engine/Source/Runtime/RenderCore/Public/ShaderCompilerJobTypes.h
--- a/Engine/Source/Runtime/RenderCore/Public/ShaderCompilerJobTypes.h
+++ b/Engine/Source/Runtime/RenderCore/Public/ShaderCompilerJobTypes.h
@@ -264,7 +264,14 @@ public:
 	// Serializes only the subset of data written by ShaderCompiler and read from SCW when using worker processes.
 	RENDERCORE_API void SerializeWorkerInput(FArchive& Ar);
 
-	RENDERCORE_API const FString& GetFinalSource() const;
+	UE_DEPRECATED(5.4, "GetFinalSource is deprecated, GetFinalSourceView returns an FStringView instead.")
+	inline const FString& GetFinalSource() const
+	{
+		static FString Empty;
+		return Empty;
+	}
+
+	RENDERCORE_API FStringView GetFinalSourceView() const;
 
 	FShaderCompileJob() : FShaderCommonCompileJob(Type, 0u, 0u, EShaderCompileJobPriority::Num)
 	{}
diff --git a/Engine/Source/Runtime/RenderCore/Public/ShaderPreprocessTypes.h b/Engine/Source/Runtime/RenderCore/Public/ShaderPreprocessTypes.h
--- a/Engine/Source/Runtime/RenderCore/Public/ShaderPreprocessTypes.h
+++ b/Engine/Source/Runtime/RenderCore/Public/ShaderPreprocessTypes.h
@@ -2,10 +2,12 @@
 
 #pragma once
 
+#include "Containers/AnsiString.h"
 #include "Containers/Array.h"
 #include "Containers/UnrealString.h"
 #include "ShaderCompilerCore.h"
 #include "ShaderCore.h"
+#include "ShaderSource.h"
 
 struct FShaderCompilerEnvironment;
 struct FShaderCompilerError;
@@ -58,23 +60,65 @@ public:
 	FShaderPreprocessOutput()
 	{
 	}
+	
+	UE_DEPRECATED(5.4, "Use StringView accessors")
 	const FString& GetSource() const
 	{
-		return PreprocessedSource;
+		static FString Empty;
+		return Empty;
+	}
+
+	FAnsiStringView GetSourceViewAnsi() const
+	{
+		if (AnsiSource.IsEmpty())
+		{
+			AnsiSource = FAnsiString(PreprocessedSource.GetView());
+		}
+		return FAnsiStringView(AnsiSource);
 	}
 
+	FStringView GetSourceViewWide() const
+	{
+		return PreprocessedSource.GetView();
+	}
+
+	UE_DEPRECATED(5.4, "Use GetUnstrippedSourceView")
 	const FString& GetUnstrippedSource() const
+	{
+		PRAGMA_DISABLE_DEPRECATION_WARNINGS
+		return GetSource();
+		PRAGMA_ENABLE_DEPRECATION_WARNINGS
+	}
+
+	FStringView GetUnstrippedSourceView() const
 	{
 		// if the unstripped source is requested, check if the "original source" field has been populated
 		// if not then stripping hasn't occurred so there's only one preprocessed source; return it
-		return OriginalPreprocessedSource.IsEmpty() ? PreprocessedSource : OriginalPreprocessedSource;
+		return OriginalPreprocessedSource.IsEmpty() ? PreprocessedSource.GetView() : OriginalPreprocessedSource.GetView();
 	}
 
-	FString& EditSource()
+	FShaderSource& EditSource()
 	{
 		return PreprocessedSource;
 	}
 
+	inline void ForEachLine(TFunction<void(FAnsiStringView Line, int32 LineIndex)> Callback) const
+	{
+		FAnsiStringView Source = GetSourceViewAnsi();
+		int32 EndIndex = 0, LineIndex = 0, StartIndex = 0;
+		while (StartIndex < Source.Len())
+		{
+			EndIndex = Source.Find(ANSITEXTVIEW("\n"), StartIndex);
+			if (EndIndex == INDEX_NONE)
+			{
+				EndIndex = Source.Len();
+			}
+			FAnsiStringView Line(Source.GetData() + StartIndex, EndIndex - StartIndex);
+			StartIndex = EndIndex + 1;
+			Callback(Line, LineIndex++);
+		}
+	}
+
 	inline bool HasDirective(const FString& Directive) const
 	{
 		const int32 NumberOfDirectives = PragmaDirectives.Num();
@@ -186,10 +230,13 @@ private:
 	void RemapErrors(FShaderCompilerOutput& Output) const;
 
 	// Output of preprocessing; should be set by IShaderFormat::PreprocessShader
-	FString PreprocessedSource;
+	FShaderSource PreprocessedSource;
 
 	// Set by Finalize; original preprocessed source as set by IShaderFormat::PreprocessShader
-	FString OriginalPreprocessedSource;
+	FShaderSource OriginalPreprocessedSource;
+
+	// Set when GetSourceViewAnsi is called (if source is wide)
+	mutable FAnsiString AnsiSource;
 
 	// Array of errors encountered in preprocessing; should be populated by IShaderFormat::PreprocessShader
 	TArray<FShaderCompilerError> Errors;
diff --git a/Engine/Source/Runtime/RenderCore/Public/ShaderSource.h b/Engine/Source/Runtime/RenderCore/Public/ShaderSource.h
new file mode 100644
--- /dev/null
+++ b/Engine/Source/Runtime/RenderCore/Public/ShaderSource.h
@@ -0,0 +1,101 @@
+// Copyright Epic Games, Inc. All Rights Reserved.
+
+#pragma once
+
+#include "Containers/AnsiString.h"
+#include "Containers/Array.h"
+#include "Containers/StringView.h"
+#include "Containers/UnrealString.h"
+
+#define SHADER_SOURCE_LITERAL(S) TEXT(S)
+#define SHADER_SOURCE_VIEWLITERAL(S) TEXTVIEW(S)
+
+/* Class used in shader compilation pipelines which wraps source code and ensures sufficient padding such that 16-byte wide SIMD
+ * operations on the source are guaranteed to read valid memory even if starting from the last character.
+ */
+class FShaderSource
+{
+private:
+	inline void SetLen(int32 Num, bool bShrink = true)
+	{
+		Source.SetNumUninitialized(Num + ShaderSourceSimdPadding, bShrink);
+		FMemory::Memzero(Source.GetData() + Num, sizeof(CharType) * ShaderSourceSimdPadding);
+	}
+
+public:
+	typedef TCHAR CharType;
+	typedef FStringView FViewType;
+	typedef FString FStringType;
+	typedef FCString FCStringType;
+
+	/* Construct an empty shader source object; will still contain padding */
+	FShaderSource() { SetLen(0); };
+
+	/* Given a string view construct a shader source object containing the contents of that string.
+	 * Note that this will incur a memcpy of the string contents.
+	 * @param InSrc The source string to be copied
+	 * @param AdditionalSlack optional additional space to allocate; this is on top of the automatic padding.
+	 */
+	RENDERCORE_API FShaderSource(FViewType InSrc, int32 AdditionalSlack = 0);
+
+	/* Set the given string as the contents of this shader source object. The inner allocation will grow to fit
+	 * the string contents as needed.
+	 * Note that this will incur a memcpy of the string contents.
+	 * @param InSrc The source string to be copied
+	 * @param AdditionalSlack optional additional space to allocate; this is on top of the automatic padding.
+	 */
+	RENDERCORE_API void Set(FViewType InSrc, int32 AdditionalSlack = 0);
+
+	/* Set the given ANSI string as the contents of this shader source object. The inner allocation will grow to fit
+	 * the string contents as needed. Note that this will incur a conversion to TCHAR and memcpy of the string contents.
+	 * @param InSrc The source ANSI string to be copied
+	 */
+	RENDERCORE_API void Set(FAnsiStringView InSrc);
+
+	/* Move assignment operator accepting a string object. This will append padding bytes to the existing string, as such it's
+	 * best if there's sufficient extra capacity in the string storage to avoid incurring a realloc-and-copy here.
+	 * @param InSrc The source string whose data this object will take ownership of.
+	 */
+	RENDERCORE_API FShaderSource& operator=(FStringType&& InSrc);
+
+	/* Reduces the set size of the stored string length, optionally shrinking the allocation.
+	 * @param Num the desired allocation size (padding bytes will be added on top of this)
+	 * @param bRealloc whether to reallocate or keep the existing larger size allocation
+	 */
+	inline void ShrinkToLen(int32 Num, bool bShrink = true)
+	{
+		checkf(Num <= Len(), TEXT("Trying to shrink to %d characters but existing allocation is smaller (%d characters)"), Num, Len());
+		SetLen(Num, bShrink);
+	}
+
+	/* String view accessor.
+	 * @return a string view pointing to the source contents, excluding padding.
+	 */
+	inline FViewType GetView() const { return { Source.GetData(), Len() }; }
+
+	/* Direct data pointer accessor
+	 * @return a pointer to the source data; will be null terminated by the SIMD padding.
+	 */
+	inline CharType* GetData() { return Source.GetData(); }
+
+	/* IsEmpty predicate
+	 * @return true if this source object is empty excluding the SIMD padding, false otherwise.
+	 */
+	inline bool IsEmpty() const { return Len() == 0; }
+
+	/* Length accessor.
+	 * @return the non-padded length of the source (also excluding null terminator)
+	 */
+	inline int32 Len() const { return Source.Num() - ShaderSourceSimdPadding; };
+
+	/* FArchive serialization operator. Note this currently serializes padding for simplicity's sake.
+	 * @param Ar The archive to serialize from/to
+	 * @param ShaderSource the source object to serialize
+	 */
+	friend FArchive& operator<<(FArchive& Ar, FShaderSource& ShaderSource);
+	
+private:
+
+	static constexpr int32 ShaderSourceSimdPadding = 7;
+	TArray<CharType> Source;
+};
