[shaders] add support for single byte char path for minifier simd optimizations

diff --git a/Engine/Source/Developer/ShaderCompilerCommon/Private/ShaderMinifier.cpp b/Engine/Source/Developer/ShaderCompilerCommon/Private/ShaderMinifier.cpp
--- a/Engine/Source/Developer/ShaderCompilerCommon/Private/ShaderMinifier.cpp
+++ b/Engine/Source/Developer/ShaderCompilerCommon/Private/ShaderMinifier.cpp
@@ -240,6 +240,27 @@ static FShaderSource::FViewType ExtractOperator(FShaderSource::FViewType Source)
 	return Result;
 }
 
+#if UE_SHADER_MINIFIER_SSE
+template <int CompareType>
+inline int32 ScanPastCharactersSimd(const FShaderSource::CharType* Source, int32 SourceLen, __m128i NeedleVec)
+{
+	constexpr int32 Mode = (FShaderSource::IsWide()  ? _SIDD_UWORD_OPS : _SIDD_UBYTE_OPS) | CompareType | _SIDD_MASKED_NEGATIVE_POLARITY;
+	int32 Cursor = 0;
+	while (Cursor < SourceLen)
+	{
+		__m128i Chunk = _mm_loadu_si128(reinterpret_cast<const __m128i*>(Source + Cursor));
+		const int32 CompareResult = _mm_cmpistrc(NeedleVec, Chunk, Mode);
+		if (CompareResult)
+		{
+			Cursor += _mm_cmpistri(NeedleVec, Chunk, Mode);
+			break;
+		}
+		Cursor += FShaderSource::GetSimdCharCount();
+	}
+	return Cursor;
+}
+#endif // UE_SHADER_MINIFIER_SSE
+
 static FShaderSource::FViewType SkipUntilNonIdentifierCharacter(FShaderSource::FViewType Source)
 {
 	const int32 SourceLen = Source.Len();
@@ -247,34 +268,21 @@ static FShaderSource::FViewType SkipUntilNonIdentifierCharacter(FShaderSource::F
 	const FShaderSource::CharType* SourceData = Source.GetData();
 
 #if UE_SHADER_MINIFIER_SSE
-	if constexpr (sizeof(FShaderSource::CharType) == 2)
+	const __m128i NeedleVec = FShaderSource::IsWide()
+		? _mm_setr_epi16(L'0', L'9', L'a', L'z', L'A', L'Z', L'_', L'_')
+		: _mm_setr_epi8('0', '9', 'a', 'z', 'A', 'Z', '_', '_', 0, 0, 0, 0, 0, 0, 0, 0);
+		
+	Cursor = ScanPastCharactersSimd<_SIDD_CMP_RANGES>(SourceData, SourceLen, NeedleVec);
+#else
+	while (Cursor < SourceLen)
 	{
-		const __m128i NeedleVec = _mm_setr_epi16(L'0', L'9', L'a', L'z', L'A', L'Z', L'_', L'_');
-		while (Cursor < SourceLen)
+		if (!IsPossibleIdentifierCharacter(SourceData[Cursor]))
 		{
-			__m128i Chunk = _mm_loadu_si128(reinterpret_cast<const __m128i*>(SourceData + Cursor));
-			constexpr int32 Mode = _SIDD_UWORD_OPS | _SIDD_CMP_RANGES | _SIDD_MASKED_NEGATIVE_POLARITY;
-			const int32 CompareResult = _mm_cmpistrc(NeedleVec, Chunk, Mode);
-			if (CompareResult)
-			{
-				Cursor += _mm_cmpistri(NeedleVec, Chunk, Mode);
-				return FShaderSource::FViewType(SourceData + Cursor, SourceLen - Cursor);
-			}
-			Cursor += 8;
+			break;
 		}
+		++Cursor;
 	}
-	else
 #endif // UE_SHADER_MINIFIER_SSE
-	{
-		while (Cursor < SourceLen)
-		{
-			if (!IsPossibleIdentifierCharacter(SourceData[Cursor]))
-			{
-				break;
-			}
-			++Cursor;
-		}
-	}
 	return FShaderSource::FViewType(SourceData + Cursor, FMath::Max(SourceLen - Cursor, 0));
 }
 
@@ -301,34 +309,21 @@ static FShaderSource::FViewType SkipSpace(FShaderSource::FViewType Source)
 	const FShaderSource::CharType* SourceData = Source.GetData();
 
 #if UE_SHADER_MINIFIER_SSE
-	if constexpr (sizeof(FShaderSource::CharType) == 2)
+	const __m128i NeedleVec = FShaderSource::IsWide()
+		? _mm_setr_epi16(L' ', L'\f', L'\r', L'\n', L'\t', L'\v', 0, 0)
+		: _mm_setr_epi8(' ', '\f', '\r', '\n', '\t', '\v', 0, 0, 0, 0, 0, 0, 0, 0, 0, 0);
+
+	Cursor = ScanPastCharactersSimd<_SIDD_CMP_EQUAL_ANY>(SourceData, SourceLen, NeedleVec);
+#else
+	while (Cursor < SourceLen)
 	{
-		const __m128i NeedleVec = _mm_setr_epi16(L' ', L'\f', L'\r', L'\n', L'\t', L'\v', 0, 0);
-		while (Cursor < SourceLen)
+		if (!IsSpace(SourceData[Cursor]))
 		{
-			__m128i Chunk = _mm_loadu_si128(reinterpret_cast<const __m128i*>(SourceData + Cursor));
-			constexpr int32 Mode = _SIDD_UWORD_OPS | _SIDD_CMP_EQUAL_ANY | _SIDD_MASKED_NEGATIVE_POLARITY;
-			const int32 CompareResult = _mm_cmpistrc(NeedleVec, Chunk, Mode);
-			if (CompareResult)
-			{
-				Cursor += _mm_cmpistri(NeedleVec, Chunk, Mode);
-				return FShaderSource::FViewType(SourceData + Cursor, SourceLen - Cursor);
-			}
-			Cursor += 8;
+			break;
 		}
+		++Cursor;
 	}
-	else
 #endif // UE_SHADER_MINIFIER_SSE
-	{
-		while (Cursor < SourceLen)
-		{
-			if (!IsSpace(SourceData[Cursor]))
-			{
-				break;
-			}
-			++Cursor;
-		}
-	}
 	return FShaderSource::FViewType(SourceData + Cursor, FMath::Max(SourceLen - Cursor, 0));
 }
 
@@ -427,39 +422,36 @@ static FShaderSource::FViewType ExtractBlock(FShaderSource::FViewType Source, FS
 	};
 
 #if UE_SHADER_MINIFIER_SSE
-	if constexpr (sizeof(FShaderSource::CharType) == 2)
+	const __m128i NeedleVec = FShaderSource::IsWide()
+		? _mm_setr_epi16(DelimBegin, DelimEnd, 0, 0, 0, 0, 0, 0)
+		: _mm_setr_epi8(DelimBegin, DelimEnd, 0, 0, 0, 0, 0, 0, 0, 0, 0, 0, 0, 0, 0, 0);
+
+	while (Cursor < SourceLen && Status != EStatus::Finished)
 	{
-		const __m128i NeedleVec = _mm_setr_epi16(DelimBegin, DelimEnd, 0, 0, 0, 0, 0, 0);
-		while (Cursor < SourceLen && Status != EStatus::Finished)
+		__m128i Chunk = _mm_loadu_si128(reinterpret_cast<const __m128i*>(SourceData + Cursor));
+		constexpr int32 Mode = (FShaderSource::IsWide() ? _SIDD_UWORD_OPS : _SIDD_UBYTE_OPS) | _SIDD_CMP_EQUAL_ANY | _SIDD_MOST_SIGNIFICANT;
+		const int32 CompareResult = _mm_cmpistrc(NeedleVec, Chunk, Mode);
+		if (CompareResult)
 		{
-			__m128i Chunk = _mm_loadu_si128(reinterpret_cast<const __m128i*>(SourceData + Cursor));
-			constexpr int32 Mode = _SIDD_UWORD_OPS | _SIDD_CMP_EQUAL_ANY | _SIDD_MOST_SIGNIFICANT;
-			const int32 CompareResult = _mm_cmpistrc(NeedleVec, Chunk, Mode);
-			if (CompareResult)
+			__m128i MaskVec = _mm_cmpistrm(NeedleVec, Chunk, Mode);
+			uint32 Mask = _mm_movemask_epi8(MaskVec);
+			while (Mask != 0 && Status != EStatus::Finished)
 			{
-				__m128i MaskVec = _mm_cmpistrm(NeedleVec, Chunk, Mode);
-				uint32 Mask = _mm_movemask_epi8(MaskVec);
-				while (Mask != 0 && Status != EStatus::Finished)
-				{
-					const uint32 BitIndex = FMath::CountTrailingZeros(Mask);
-					const uint32 ChunkCharIndex = BitIndex / 2;
-					Status = ProcessCharacter(Cursor + ChunkCharIndex);
-					Mask &= ~(3 << BitIndex);
-				}
+				const uint32 BitIndex = FMath::CountTrailingZeros(Mask);
+				const uint32 ChunkCharIndex = BitIndex / sizeof(FShaderSource::CharType);
+				Status = ProcessCharacter(Cursor + ChunkCharIndex);
+				Mask &= ~(FShaderSource::GetSingleCharMask() << BitIndex);
 			}
-			Cursor += 8;
 		}
+		Cursor += FShaderSource::GetSimdCharCount();
 	}
-	else
-#endif // UE_SHADER_MINIFIER_SSE
+#else
+	while (Cursor < SourceLen && Status != EStatus::Finished)
 	{
-		while (Cursor < SourceLen && Status != EStatus::Finished)
-		{
-			Status = ProcessCharacter(Cursor);
-			++Cursor;
-		}
+		Status = ProcessCharacter(Cursor);
+		++Cursor;
 	}
-
+#endif // UE_SHADER_MINIFIER_SSE
 	if (Stack == 0 && PosEnd != INDEX_NONE)
 	{
 		return FShaderSource::FViewType(Source.GetData(), PosEnd + 1);
@@ -1474,37 +1466,31 @@ static void BuildLineBreakMap(FShaderSource::FViewType Source, TArray<int32>& Ou
 	int32 Cursor = 0;
 
 #if UE_SHADER_MINIFIER_SSE
-	if constexpr (sizeof(FShaderSource::CharType) == 2)
+	const __m128i Needle = FShaderSource::IsWide() ? _mm_set1_epi16(L'\n') : _mm_set1_epi8('\n');
+	while (Cursor < SourceLen)
 	{
-		const __m128i Needle = _mm_set1_epi16(L'\n');
-		while (Cursor < SourceLen)
+		__m128i Chunk = _mm_loadu_si128(reinterpret_cast<const __m128i*>(Chars + Cursor));
+		__m128i MaskVec = FShaderSource::IsWide() ? _mm_cmpeq_epi16(Chunk, Needle) : _mm_cmpeq_epi8(Chunk, Needle);
+		uint32 Mask = _mm_movemask_epi8(MaskVec);
+		while (Mask != 0)
 		{
-			__m128i Chunk = _mm_loadu_si128(reinterpret_cast<const __m128i*>(Chars + Cursor));
-			__m128i MaskVec = _mm_cmpeq_epi16(Chunk, Needle);
-			uint32 Mask = _mm_movemask_epi8(MaskVec);
-			while (Mask != 0)
-			{
-				// NOTE: 2 bits represent each character
-				const uint32 BitIndex = FMath::CountTrailingZeros(Mask);
-				const uint32 ChunkCharIndex = BitIndex / 2;
-				OutLineBreakMap.Add(Cursor + ChunkCharIndex);
-				Mask &= ~(3 << BitIndex);
-			}
-			Cursor += 8;
+			const uint32 BitIndex = FMath::CountTrailingZeros(Mask);
+			const uint32 ChunkCharIndex = BitIndex / sizeof(FShaderSource::CharType);
+			OutLineBreakMap.Add(Cursor + ChunkCharIndex);
+			Mask &= ~(FShaderSource::GetSingleCharMask() << BitIndex);
 		}
+		Cursor += FShaderSource::GetSimdCharCount();
 	}
-	else
-#endif //UE_SHADER_MINIFIER_SSE
+#else
+	while (Cursor < SourceLen)
 	{
-		while (Cursor < SourceLen)
+		if (Chars[Cursor] == FShaderSource::CharType('\n'))
 		{
-			if (Chars[Cursor] == FShaderSource::CharType('\n'))
-			{
-				OutLineBreakMap.Add(Cursor);
-			}
-			++Cursor;
+			OutLineBreakMap.Add(Cursor);
 		}
+		++Cursor;
 	}
+#endif //UE_SHADER_MINIFIER_SSE
 }
 
 static int32 FindLineDirective(const TArray<FShaderSource::FViewType>& LineDirectives, const FShaderSource::CharType* Ptr)
diff --git a/Engine/Source/Runtime/RenderCore/Public/ShaderSource.h b/Engine/Source/Runtime/RenderCore/Public/ShaderSource.h
--- a/Engine/Source/Runtime/RenderCore/Public/ShaderSource.h
+++ b/Engine/Source/Runtime/RenderCore/Public/ShaderSource.h
@@ -28,6 +28,13 @@ public:
 	typedef FString FStringType;
 	typedef FCString FCStringType;
 
+	/** Constexpr predicate indicating whether wide or ansi chars are used */
+	static constexpr bool IsWide() { return sizeof(CharType) == 2; }
+	/** Constexpr function returning the number of characters read in a single SIMD compare op */
+	static constexpr int32 GetSimdCharCount() { return IsWide() ? 8 : 16; }
+	/** Constexpr function returning a mask value for a single character */
+	static constexpr int32 GetSingleCharMask() { return IsWide() ? 3 : 1; }
+
 	/* Construct an empty shader source object; will still contain padding */
 	FShaderSource() { SetLen(0); };
 
