RDNA disassembler fixes
**
- Support for DS instructions
- Fixed DS microcode format
- Added missing instructions

diff --git a/Engine/Source/Developer/ShaderCompilerCommon/Private/ISAParser/QuadModeDetect.cpp b/Engine/Source/Developer/ShaderCompilerCommon/Private/ISAParser/QuadModeDetect.cpp
--- a/Engine/Source/Developer/ShaderCompilerCommon/Private/ISAParser/QuadModeDetect.cpp
+++ b/Engine/Source/Developer/ShaderCompilerCommon/Private/ISAParser/QuadModeDetect.cpp
@@ -628,8 +628,15 @@ bool PerformParsing(PARSER& Parser, const uint32* ReadPtr, const uint32* EndPtr)
 		EInstructionType InstructionType = DecodeInstructionType(ReadPtr, AdvanceAmount);
 		AdvanceAmount += Parser.AdjustAdvance(InstructionType, ReadPtr);
 
-		if ((ReadPtr + AdvanceAmount >= EndPtr) || InstructionType == EInstructionType::UNKNOWN)
+		if (InstructionType == EInstructionType::UNKNOWN)
 		{
+			//UE_LOG(LogTemp, Error, TEXT("Failed ISA Parse - Unknown Instruction"));
+			bSuccess = false;
+			break;
+		}
+		else if (ReadPtr + AdvanceAmount >= EndPtr)
+		{
+			//UE_LOG(LogTemp, Error, TEXT("Failed ISA Parse - Overfetch"));
 			bSuccess = false;
 			break;
 		}
@@ -702,6 +709,8 @@ bool PerformParsing(PARSER& Parser, const uint32* ReadPtr, const uint32* EndPtr)
 	return bSuccess;
 }
 
+#define DEBUG_DUMP_FAILURES 0
+
 bool ISAParser::HasDerivativeOps(bool& bHasDerivativeOps, const char* Code, uint32 CodeLength, EInstructionSet InstructionSet)
 {
 	const uint32* ReadPtr = reinterpret_cast<const uint32*>(Code);
@@ -712,9 +721,6 @@ bool ISAParser::HasDerivativeOps(bool& bHasDerivativeOps, const char* Code, uint
 
 	if (InstructionSet == EInstructionSet::RDNA1)
 	{
-		//FDebugParserRDNA1 PrintParser;
-		//PerformParsing(PrintParser, ReadPtr, EndPtr);
-
 		FQuadModeParserRDNA1 Parser;
 		if (PerformParsing(Parser, ReadPtr, EndPtr))
 		{
@@ -722,14 +728,16 @@ bool ISAParser::HasDerivativeOps(bool& bHasDerivativeOps, const char* Code, uint
 		}
 		else
 		{
+		#if DEBUG_DUMP_FAILURES
+			FDebugParserRDNA1 PrintParser;
+			PerformParsing(PrintParser, ReadPtr, EndPtr);
+		#endif
+
 			return false;
 		}
 	}
 	else if (InstructionSet == EInstructionSet::RDNA2)
 	{
-		//FDebugParserRDNA2 PrintParser;
-		//PerformParsing(PrintParser, ReadPtr, EndPtr);
-
 		FQuadModeParserRDNA2 Parser;
 		if (PerformParsing(Parser, ReadPtr, EndPtr))
 		{
@@ -737,6 +745,11 @@ bool ISAParser::HasDerivativeOps(bool& bHasDerivativeOps, const char* Code, uint
 		}
 		else
 		{
+		#if DEBUG_DUMP_FAILURES
+			FDebugParserRDNA2 PrintParser;
+			PerformParsing(PrintParser, ReadPtr, EndPtr);
+		#endif
+
 			return false;
 		}
 	}
diff --git a/Engine/Source/Developer/ShaderCompilerCommon/Private/ISAParser/RDNA1_ISA.h b/Engine/Source/Developer/ShaderCompilerCommon/Private/ISAParser/RDNA1_ISA.h
--- a/Engine/Source/Developer/ShaderCompilerCommon/Private/ISAParser/RDNA1_ISA.h
+++ b/Engine/Source/Developer/ShaderCompilerCommon/Private/ISAParser/RDNA1_ISA.h
@@ -188,6 +188,7 @@ enum class ESOP2Ops : uint16
 	s_nor_b64 = 27,
 	s_xnor_b32 = 28,
 	s_xnor_b64 = 29,
+	s_lshl_b32 = 30,
 	s_lshl_b64 = 31,
 	s_lshr_b32 = 32,
 	s_lshr_b64 = 33,
@@ -246,6 +247,7 @@ const char* ToString(ESOP2Ops Op)
 		OP_TO_STRING_CASE(s_nor_b64);
 		OP_TO_STRING_CASE(s_xnor_b32);
 		OP_TO_STRING_CASE(s_xnor_b64);
+		OP_TO_STRING_CASE(s_lshl_b32);
 		OP_TO_STRING_CASE(s_lshl_b64);
 		OP_TO_STRING_CASE(s_lshr_b32);
 		OP_TO_STRING_CASE(s_lshr_b64);
@@ -1571,7 +1573,7 @@ const char* ToString(EVINTERPOps Op)
 		OP_TO_STRING_CASE(v_interp_p2_f32);
 		OP_TO_STRING_CASE(v_interp_mov_f32);
 	default:
-		return "UNKNOWN - VOP3 A/B";
+		return "UNKNOWN - VINTERP";
 	}
 
 #undef OP_TO_STRING_CASE
@@ -2182,6 +2184,269 @@ const char* ToString(ESCRATCHOps Op)
 #undef OP_TO_STRING_CASE
 }
 
+enum class ELDSGDSOps : uint16
+{
+	ds_add_u32 = 0,
+	ds_sub_u32 = 1,
+	ds_rsub_u32 = 2,
+	ds_inc_u32 = 3,
+	ds_dec_u32 = 4,
+	ds_min_i32 = 5,
+	ds_max_i32 = 6,
+	ds_min_u32 = 7,
+	ds_max_u32 = 8,
+	ds_and_b32 = 9,
+	ds_or_b32 = 10,
+	ds_xor_b32 = 11,
+	ds_mskor_b32 = 12,
+	ds_write_b32 = 13,
+	ds_write2_b32 = 14,
+	ds_write2st64_b32 = 15,
+	ds_cmpst_b32 = 16,
+	ds_cmpst_f32 = 17,
+	ds_min_f32 = 18,
+	ds_max_f32 = 19,
+	ds_nop = 20,
+	ds_add_f32 = 21,
+	ds_gws_sema_release_all = 24,
+	ds_gws_init = 25,
+	ds_gws_sema_v = 26,
+	ds_gws_sema_br = 27,
+	ds_gws_sema_p = 28,
+	ds_gws_barrier = 29,
+	ds_write_b8 = 30,
+	ds_write_b16 = 31,
+	ds_add_rtn_u32 = 32,
+	ds_sub_rtn_u32 = 33,
+	ds_rsub_rtn_u32 = 34,
+	ds_inc_rtn_u32 = 35,
+	ds_dec_rtn_u32 = 36,
+	ds_min_rtn_i32 = 37,
+	ds_max_rtn_i32 = 38,
+	ds_min_rtn_u32 = 39,
+	ds_max_rtn_u32 = 40,
+	ds_and_rtn_b32 = 41,
+	ds_or_rtn_b32 = 42,
+	ds_xor_rtn_b32 = 43,
+	ds_mskor_rtn_b32 = 44,
+	ds_wrxchg_rtn_b32 = 45,
+	ds_wrxchg2_rtn_b32 = 46,
+	ds_wrxchg2st64_rtn_b32 = 47,
+	ds_cmpst_rtn_b32 = 48,
+	ds_cmpst_rtn_f32 = 49,
+	ds_min_rtn_f32 = 50,
+	ds_max_rtn_f32 = 51,
+	ds_wrap_rtn_b32 = 52,
+	ds_swizzle_b32 = 53,
+	ds_read_b32 = 54,
+	ds_read2_b32 = 55,
+	ds_read2st64_b32 = 56,
+	ds_read_i8 = 57,
+	ds_read_u8 = 58,
+	ds_read_i16 = 59,
+	ds_read_u16 = 60,
+	ds_consume = 61,
+	ds_append = 62,
+	ds_ordered_count = 63,
+	ds_add_u64 = 64,
+	ds_sub_u64 = 65,
+	ds_rsub_u64 = 66,
+	ds_inc_u64 = 67,
+	ds_dec_u64 = 68,
+	ds_min_i64 = 69,
+	ds_max_i64 = 70,
+	ds_min_u64 = 71,
+	ds_max_u64 = 72,
+	ds_and_b64 = 73,
+	ds_or_b64 = 74,
+	ds_xor_b64 = 75,
+	ds_mskor_b64 = 76,
+	ds_write_b64 = 77,
+	ds_write2_b64 = 78,
+	ds_write2st64_b64 = 79,
+	ds_cmpst_b64 = 80,
+	ds_cmpst_f64 = 81,
+	ds_min_f64 = 82,
+	ds_max_f64 = 83,
+	ds_add_rtn_f32 = 85,
+	ds_add_rtn_u64 = 96,
+	ds_sub_rtn_u64 = 97,
+	ds_rsub_rtn_u64 = 98,
+	ds_inc_rtn_u64 = 99,
+	ds_dec_rtn_u64 = 100,
+	ds_min_rtn_i64 = 101,
+	ds_max_rtn_i64 = 102,
+	ds_min_rtn_u64 = 103,
+	ds_max_rtn_u64 = 104,
+	ds_and_rtn_b64 = 105,
+	ds_or_rtn_b64 = 106,
+	ds_xor_rtn_b64 = 107,
+	ds_mskor_rtn_b64 = 108,
+	ds_wrxchg_rtn_b64 = 109,
+	ds_wrxchg2_rtn_b64 = 110,
+	ds_wrxchg2st64_rtn_b64 = 111,
+	ds_cmpst_rtn_b64 = 112,
+	ds_cmpst_rtn_f64 = 113,
+	ds_min_rtn_f64 = 114,
+	ds_max_rtn_f64 = 115,
+	ds_read_b64 = 118,
+	ds_read2_b64 = 119,
+	ds_read2st64_b64 = 120,
+	ds_condxchg32_rtn_b64 = 126,
+	ds_write_b8_d16_hi = 160,
+	ds_write_b16_d16_hi = 161,
+	ds_read_u8_d16 = 162,
+	ds_read_u8_d16_hi = 163,
+	ds_read_i8_d16 = 164,
+	ds_read_i8_d16_hi = 165,
+	ds_read_u16_d16 = 166,
+	ds_read_u16_d16_hi = 167,
+	ds_write_addtid_b32 = 176,
+	ds_read_addtid_b32 = 177,
+	ds_permute_b32 = 178,
+	ds_bpermute_b32 = 179,
+	ds_write_b96 = 222,
+	ds_write_b128 = 223,
+	ds_read_b96 = 254,
+	ds_read_b128 = 255,
+};
+
+const char* ToString(ELDSGDSOps Op)
+{
+#define OP_TO_STRING_CASE(x) case ELDSGDSOps::x: return #x
+
+	switch (Op)
+	{
+		OP_TO_STRING_CASE(ds_add_u32);
+		OP_TO_STRING_CASE(ds_sub_u32);
+		OP_TO_STRING_CASE(ds_rsub_u32);
+		OP_TO_STRING_CASE(ds_inc_u32);
+		OP_TO_STRING_CASE(ds_dec_u32);
+		OP_TO_STRING_CASE(ds_min_i32);
+		OP_TO_STRING_CASE(ds_max_i32);
+		OP_TO_STRING_CASE(ds_min_u32);
+		OP_TO_STRING_CASE(ds_max_u32);
+		OP_TO_STRING_CASE(ds_and_b32);
+		OP_TO_STRING_CASE(ds_or_b32);
+		OP_TO_STRING_CASE(ds_xor_b32);
+		OP_TO_STRING_CASE(ds_mskor_b32);
+		OP_TO_STRING_CASE(ds_write_b32);
+		OP_TO_STRING_CASE(ds_write2_b32);
+		OP_TO_STRING_CASE(ds_write2st64_b32);
+		OP_TO_STRING_CASE(ds_cmpst_b32);
+		OP_TO_STRING_CASE(ds_cmpst_f32);
+		OP_TO_STRING_CASE(ds_min_f32);
+		OP_TO_STRING_CASE(ds_max_f32);
+		OP_TO_STRING_CASE(ds_nop);
+		OP_TO_STRING_CASE(ds_add_f32);
+		OP_TO_STRING_CASE(ds_gws_sema_release_all);
+		OP_TO_STRING_CASE(ds_gws_init);
+		OP_TO_STRING_CASE(ds_gws_sema_v);
+		OP_TO_STRING_CASE(ds_gws_sema_br);
+		OP_TO_STRING_CASE(ds_gws_sema_p);
+		OP_TO_STRING_CASE(ds_gws_barrier);
+		OP_TO_STRING_CASE(ds_write_b8);
+		OP_TO_STRING_CASE(ds_write_b16);
+		OP_TO_STRING_CASE(ds_add_rtn_u32);
+		OP_TO_STRING_CASE(ds_sub_rtn_u32);
+		OP_TO_STRING_CASE(ds_rsub_rtn_u32);
+		OP_TO_STRING_CASE(ds_inc_rtn_u32);
+		OP_TO_STRING_CASE(ds_dec_rtn_u32);
+		OP_TO_STRING_CASE(ds_min_rtn_i32);
+		OP_TO_STRING_CASE(ds_max_rtn_i32);
+		OP_TO_STRING_CASE(ds_min_rtn_u32);
+		OP_TO_STRING_CASE(ds_max_rtn_u32);
+		OP_TO_STRING_CASE(ds_and_rtn_b32);
+		OP_TO_STRING_CASE(ds_or_rtn_b32);
+		OP_TO_STRING_CASE(ds_xor_rtn_b32);
+		OP_TO_STRING_CASE(ds_mskor_rtn_b32);
+		OP_TO_STRING_CASE(ds_wrxchg_rtn_b32);
+		OP_TO_STRING_CASE(ds_wrxchg2_rtn_b32);
+		OP_TO_STRING_CASE(ds_wrxchg2st64_rtn_b32);
+		OP_TO_STRING_CASE(ds_cmpst_rtn_b32);
+		OP_TO_STRING_CASE(ds_cmpst_rtn_f32);
+		OP_TO_STRING_CASE(ds_min_rtn_f32);
+		OP_TO_STRING_CASE(ds_max_rtn_f32);
+		OP_TO_STRING_CASE(ds_wrap_rtn_b32);
+		OP_TO_STRING_CASE(ds_swizzle_b32);
+		OP_TO_STRING_CASE(ds_read_b32);
+		OP_TO_STRING_CASE(ds_read2_b32);
+		OP_TO_STRING_CASE(ds_read2st64_b32);
+		OP_TO_STRING_CASE(ds_read_i8);
+		OP_TO_STRING_CASE(ds_read_u8);
+		OP_TO_STRING_CASE(ds_read_i16);
+		OP_TO_STRING_CASE(ds_read_u16);
+		OP_TO_STRING_CASE(ds_consume);
+		OP_TO_STRING_CASE(ds_append);
+		OP_TO_STRING_CASE(ds_ordered_count);
+		OP_TO_STRING_CASE(ds_add_u64);
+		OP_TO_STRING_CASE(ds_sub_u64);
+		OP_TO_STRING_CASE(ds_rsub_u64);
+		OP_TO_STRING_CASE(ds_inc_u64);
+		OP_TO_STRING_CASE(ds_dec_u64);
+		OP_TO_STRING_CASE(ds_min_i64);
+		OP_TO_STRING_CASE(ds_max_i64);
+		OP_TO_STRING_CASE(ds_min_u64);
+		OP_TO_STRING_CASE(ds_max_u64);
+		OP_TO_STRING_CASE(ds_and_b64);
+		OP_TO_STRING_CASE(ds_or_b64);
+		OP_TO_STRING_CASE(ds_xor_b64);
+		OP_TO_STRING_CASE(ds_mskor_b64);
+		OP_TO_STRING_CASE(ds_write_b64);
+		OP_TO_STRING_CASE(ds_write2_b64);
+		OP_TO_STRING_CASE(ds_write2st64_b64);
+		OP_TO_STRING_CASE(ds_cmpst_b64);
+		OP_TO_STRING_CASE(ds_cmpst_f64);
+		OP_TO_STRING_CASE(ds_min_f64);
+		OP_TO_STRING_CASE(ds_max_f64);
+		OP_TO_STRING_CASE(ds_add_rtn_f32);
+		OP_TO_STRING_CASE(ds_add_rtn_u64);
+		OP_TO_STRING_CASE(ds_sub_rtn_u64);
+		OP_TO_STRING_CASE(ds_rsub_rtn_u64);
+		OP_TO_STRING_CASE(ds_inc_rtn_u64);
+		OP_TO_STRING_CASE(ds_dec_rtn_u64);
+		OP_TO_STRING_CASE(ds_min_rtn_i64);
+		OP_TO_STRING_CASE(ds_max_rtn_i64);
+		OP_TO_STRING_CASE(ds_min_rtn_u64);
+		OP_TO_STRING_CASE(ds_max_rtn_u64);
+		OP_TO_STRING_CASE(ds_and_rtn_b64);
+		OP_TO_STRING_CASE(ds_or_rtn_b64);
+		OP_TO_STRING_CASE(ds_xor_rtn_b64);
+		OP_TO_STRING_CASE(ds_mskor_rtn_b64);
+		OP_TO_STRING_CASE(ds_wrxchg_rtn_b64);
+		OP_TO_STRING_CASE(ds_wrxchg2_rtn_b64);
+		OP_TO_STRING_CASE(ds_wrxchg2st64_rtn_b64);
+		OP_TO_STRING_CASE(ds_cmpst_rtn_b64);
+		OP_TO_STRING_CASE(ds_cmpst_rtn_f64);
+		OP_TO_STRING_CASE(ds_min_rtn_f64);
+		OP_TO_STRING_CASE(ds_max_rtn_f64);
+		OP_TO_STRING_CASE(ds_read_b64);
+		OP_TO_STRING_CASE(ds_read2_b64);
+		OP_TO_STRING_CASE(ds_read2st64_b64);
+		OP_TO_STRING_CASE(ds_condxchg32_rtn_b64);
+		OP_TO_STRING_CASE(ds_write_b8_d16_hi);
+		OP_TO_STRING_CASE(ds_write_b16_d16_hi);
+		OP_TO_STRING_CASE(ds_read_u8_d16);
+		OP_TO_STRING_CASE(ds_read_u8_d16_hi);
+		OP_TO_STRING_CASE(ds_read_i8_d16);
+		OP_TO_STRING_CASE(ds_read_i8_d16_hi);
+		OP_TO_STRING_CASE(ds_read_u16_d16);
+		OP_TO_STRING_CASE(ds_read_u16_d16_hi);
+		OP_TO_STRING_CASE(ds_write_addtid_b32);
+		OP_TO_STRING_CASE(ds_read_addtid_b32);
+		OP_TO_STRING_CASE(ds_permute_b32);
+		OP_TO_STRING_CASE(ds_bpermute_b32);
+		OP_TO_STRING_CASE(ds_write_b96);
+		OP_TO_STRING_CASE(ds_write_b128);
+		OP_TO_STRING_CASE(ds_read_b96);
+		OP_TO_STRING_CASE(ds_read_b128);
+	default:
+		return "UNKNOWN - LDS / GDS";
+	}
+
+#undef OP_TO_STRING_CASE
+}
+
 enum class EGLOBALOps : uint16
 {
 	global_load_ubyte = 8,
@@ -2423,9 +2688,8 @@ void PrintVOPC(const FInstVOPC& Inst)
 
 void PrintLDSGDS(const FInstLDSGDS& Inst)
 {
-	//ELDSGDSOps Op = (ELDSGDSOps)Inst.OP;
-	//UE_LOG(LogTemp, Warning, TEXT("%hs"), ToString(Op));
-	UE_LOG(LogTemp, Warning, TEXT("LDSGDS - todo"));
+	ELDSGDSOps Op = (ELDSGDSOps)Inst.OP;
+	UE_LOG(LogTemp, Warning, TEXT("%hs"), ToString(Op));
 }
 
 void PrintMUBUF(const FInstMUBUF& Inst)
diff --git a/Engine/Source/Developer/ShaderCompilerCommon/Private/ISAParser/RDNA2_ISA.h b/Engine/Source/Developer/ShaderCompilerCommon/Private/ISAParser/RDNA2_ISA.h
--- a/Engine/Source/Developer/ShaderCompilerCommon/Private/ISAParser/RDNA2_ISA.h
+++ b/Engine/Source/Developer/ShaderCompilerCommon/Private/ISAParser/RDNA2_ISA.h
@@ -2184,6 +2184,270 @@ const char* ToString(ESCRATCHOps Op)
 #undef OP_TO_STRING_CASE
 }
 
+enum class ELDSGDSOps : uint16
+{
+	ds_add_u32 = 0,
+	ds_sub_u32 = 1,
+	ds_rsub_u32 = 2,
+	ds_inc_u32 = 3,
+	ds_dec_u32 = 4,
+	ds_min_i32 = 5,
+	ds_max_i32 = 6,
+	ds_min_u32 = 7,
+	ds_max_u32 = 8,
+	ds_and_b32 = 9,
+	ds_or_b32 = 10,
+	ds_xor_b32 = 11,
+	ds_mskor_b32 = 12,
+	ds_write_b32 = 13,
+	ds_write2_b32 = 14,
+	ds_write2st64_b32 = 15,
+	ds_cmpst_b32 = 16,
+	ds_cmpst_f32 = 17,
+	ds_min_f32 = 18,
+	ds_max_f32 = 19,
+	ds_nop = 20,
+	ds_add_f32 = 21,
+	ds_gws_sema_release_all = 24,
+	ds_gws_init = 25,
+	ds_gws_sema_v = 26,
+	ds_gws_sema_br = 27,
+	ds_gws_sema_p = 28,
+	ds_gws_barrier = 29,
+	ds_write_b8 = 30,
+	ds_write_b16 = 31,
+	ds_add_rtn_u32 = 32,
+	ds_sub_rtn_u32 = 33,
+	ds_rsub_rtn_u32 = 34,
+	ds_inc_rtn_u32 = 35,
+	ds_dec_rtn_u32 = 36,
+	ds_min_rtn_i32 = 37,
+	ds_max_rtn_i32 = 38,
+	ds_min_rtn_u32 = 39,
+	ds_max_rtn_u32 = 40,
+	ds_and_rtn_b32 = 41,
+	ds_or_rtn_b32 = 42,
+	ds_xor_rtn_b32 = 43,
+	ds_mskor_rtn_b32 = 44,
+	ds_wrxchg_rtn_b32 = 45,
+	ds_wrxchg2_rtn_b32 = 46,
+	ds_wrxchg2st64_rtn_b32 = 47,
+	ds_cmpst_rtn_b32 = 48,
+	ds_cmpst_rtn_f32 = 49,
+	ds_min_rtn_f32 = 50,
+	ds_max_rtn_f32 = 51,
+	ds_wrap_rtn_b32 = 52,
+	ds_swizzle_b32 = 53,
+	ds_read_b32 = 54,
+	ds_read2_b32 = 55,
+	ds_read2st64_b32 = 56,
+	ds_read_i8 = 57,
+	ds_read_u8 = 58,
+	ds_read_i16 = 59,
+	ds_read_u16 = 60,
+	ds_consume = 61,
+	ds_append = 62,
+	ds_ordered_count = 63,
+	ds_add_u64 = 64,
+	ds_sub_u64 = 65,
+	ds_rsub_u64 = 66,
+	ds_inc_u64 = 67,
+	ds_dec_u64 = 68,
+	ds_min_i64 = 69,
+	ds_max_i64 = 70,
+	ds_min_u64 = 71,
+	ds_max_u64 = 72,
+	ds_and_b64 = 73,
+	ds_or_b64 = 74,
+	ds_xor_b64 = 75,
+	ds_mskor_b64 = 76,
+	ds_write_b64 = 77,
+	ds_write2_b64 = 78,
+	ds_write2st64_b64 = 79,
+	ds_cmpst_b64 = 80,
+	ds_cmpst_f64 = 81,
+	ds_min_f64 = 82,
+	ds_max_f64 = 83,
+	ds_add_rtn_f32 = 85,
+	ds_add_rtn_u64 = 96,
+	ds_sub_rtn_u64 = 97,
+	ds_rsub_rtn_u64 = 98,
+	ds_inc_rtn_u64 = 99,
+	ds_dec_rtn_u64 = 100,
+	ds_min_rtn_i64 = 101,
+	ds_max_rtn_i64 = 102,
+	ds_min_rtn_u64 = 103,
+	ds_max_rtn_u64 = 104,
+	ds_and_rtn_b64 = 105,
+	ds_or_rtn_b64 = 106,
+	ds_xor_rtn_b64 = 107,
+	ds_mskor_rtn_b64 = 108,
+	ds_wrxchg_rtn_b64 = 109,
+	ds_wrxchg2_rtn_b64 = 110,
+	ds_wrxchg2st64_rtn_b64 = 111,
+	ds_cmpst_rtn_b64 = 112,
+	ds_cmpst_rtn_f64 = 113,
+	ds_min_rtn_f64 = 114,
+	ds_max_rtn_f64 = 115,
+	ds_read_b64 = 118,
+	ds_read2_b64 = 119,
+	ds_read2st64_b64 = 120,
+	ds_condxchg32_rtn_b64 = 126,
+	ds_write_b8_d16_hi = 160,
+	ds_write_b16_d16_hi = 161,
+	ds_read_u8_d16 = 162,
+	ds_read_u8_d16_hi = 163,
+	ds_read_i8_d16 = 164,
+	ds_read_i8_d16_hi = 165,
+	ds_read_u16_d16 = 166,
+	ds_read_u16_d16_hi = 167,
+	ds_write_addtid_b32 = 176,
+	ds_read_addtid_b32 = 177,
+	ds_permute_b32 = 178,
+	ds_bpermute_b32 = 179,
+	ds_write_b96 = 222,
+	ds_write_b128 = 223,
+	ds_read_b96 = 254,
+	ds_read_b128 = 255,
+};
+
+// TODO: Compare against RDNA 2
+const char* ToString(ELDSGDSOps Op)
+{
+#define OP_TO_STRING_CASE(x) case ELDSGDSOps::x: return #x
+
+	switch (Op)
+	{
+		OP_TO_STRING_CASE(ds_add_u32);
+		OP_TO_STRING_CASE(ds_sub_u32);
+		OP_TO_STRING_CASE(ds_rsub_u32);
+		OP_TO_STRING_CASE(ds_inc_u32);
+		OP_TO_STRING_CASE(ds_dec_u32);
+		OP_TO_STRING_CASE(ds_min_i32);
+		OP_TO_STRING_CASE(ds_max_i32);
+		OP_TO_STRING_CASE(ds_min_u32);
+		OP_TO_STRING_CASE(ds_max_u32);
+		OP_TO_STRING_CASE(ds_and_b32);
+		OP_TO_STRING_CASE(ds_or_b32);
+		OP_TO_STRING_CASE(ds_xor_b32);
+		OP_TO_STRING_CASE(ds_mskor_b32);
+		OP_TO_STRING_CASE(ds_write_b32);
+		OP_TO_STRING_CASE(ds_write2_b32);
+		OP_TO_STRING_CASE(ds_write2st64_b32);
+		OP_TO_STRING_CASE(ds_cmpst_b32);
+		OP_TO_STRING_CASE(ds_cmpst_f32);
+		OP_TO_STRING_CASE(ds_min_f32);
+		OP_TO_STRING_CASE(ds_max_f32);
+		OP_TO_STRING_CASE(ds_nop);
+		OP_TO_STRING_CASE(ds_add_f32);
+		OP_TO_STRING_CASE(ds_gws_sema_release_all);
+		OP_TO_STRING_CASE(ds_gws_init);
+		OP_TO_STRING_CASE(ds_gws_sema_v);
+		OP_TO_STRING_CASE(ds_gws_sema_br);
+		OP_TO_STRING_CASE(ds_gws_sema_p);
+		OP_TO_STRING_CASE(ds_gws_barrier);
+		OP_TO_STRING_CASE(ds_write_b8);
+		OP_TO_STRING_CASE(ds_write_b16);
+		OP_TO_STRING_CASE(ds_add_rtn_u32);
+		OP_TO_STRING_CASE(ds_sub_rtn_u32);
+		OP_TO_STRING_CASE(ds_rsub_rtn_u32);
+		OP_TO_STRING_CASE(ds_inc_rtn_u32);
+		OP_TO_STRING_CASE(ds_dec_rtn_u32);
+		OP_TO_STRING_CASE(ds_min_rtn_i32);
+		OP_TO_STRING_CASE(ds_max_rtn_i32);
+		OP_TO_STRING_CASE(ds_min_rtn_u32);
+		OP_TO_STRING_CASE(ds_max_rtn_u32);
+		OP_TO_STRING_CASE(ds_and_rtn_b32);
+		OP_TO_STRING_CASE(ds_or_rtn_b32);
+		OP_TO_STRING_CASE(ds_xor_rtn_b32);
+		OP_TO_STRING_CASE(ds_mskor_rtn_b32);
+		OP_TO_STRING_CASE(ds_wrxchg_rtn_b32);
+		OP_TO_STRING_CASE(ds_wrxchg2_rtn_b32);
+		OP_TO_STRING_CASE(ds_wrxchg2st64_rtn_b32);
+		OP_TO_STRING_CASE(ds_cmpst_rtn_b32);
+		OP_TO_STRING_CASE(ds_cmpst_rtn_f32);
+		OP_TO_STRING_CASE(ds_min_rtn_f32);
+		OP_TO_STRING_CASE(ds_max_rtn_f32);
+		OP_TO_STRING_CASE(ds_wrap_rtn_b32);
+		OP_TO_STRING_CASE(ds_swizzle_b32);
+		OP_TO_STRING_CASE(ds_read_b32);
+		OP_TO_STRING_CASE(ds_read2_b32);
+		OP_TO_STRING_CASE(ds_read2st64_b32);
+		OP_TO_STRING_CASE(ds_read_i8);
+		OP_TO_STRING_CASE(ds_read_u8);
+		OP_TO_STRING_CASE(ds_read_i16);
+		OP_TO_STRING_CASE(ds_read_u16);
+		OP_TO_STRING_CASE(ds_consume);
+		OP_TO_STRING_CASE(ds_append);
+		OP_TO_STRING_CASE(ds_ordered_count);
+		OP_TO_STRING_CASE(ds_add_u64);
+		OP_TO_STRING_CASE(ds_sub_u64);
+		OP_TO_STRING_CASE(ds_rsub_u64);
+		OP_TO_STRING_CASE(ds_inc_u64);
+		OP_TO_STRING_CASE(ds_dec_u64);
+		OP_TO_STRING_CASE(ds_min_i64);
+		OP_TO_STRING_CASE(ds_max_i64);
+		OP_TO_STRING_CASE(ds_min_u64);
+		OP_TO_STRING_CASE(ds_max_u64);
+		OP_TO_STRING_CASE(ds_and_b64);
+		OP_TO_STRING_CASE(ds_or_b64);
+		OP_TO_STRING_CASE(ds_xor_b64);
+		OP_TO_STRING_CASE(ds_mskor_b64);
+		OP_TO_STRING_CASE(ds_write_b64);
+		OP_TO_STRING_CASE(ds_write2_b64);
+		OP_TO_STRING_CASE(ds_write2st64_b64);
+		OP_TO_STRING_CASE(ds_cmpst_b64);
+		OP_TO_STRING_CASE(ds_cmpst_f64);
+		OP_TO_STRING_CASE(ds_min_f64);
+		OP_TO_STRING_CASE(ds_max_f64);
+		OP_TO_STRING_CASE(ds_add_rtn_f32);
+		OP_TO_STRING_CASE(ds_add_rtn_u64);
+		OP_TO_STRING_CASE(ds_sub_rtn_u64);
+		OP_TO_STRING_CASE(ds_rsub_rtn_u64);
+		OP_TO_STRING_CASE(ds_inc_rtn_u64);
+		OP_TO_STRING_CASE(ds_dec_rtn_u64);
+		OP_TO_STRING_CASE(ds_min_rtn_i64);
+		OP_TO_STRING_CASE(ds_max_rtn_i64);
+		OP_TO_STRING_CASE(ds_min_rtn_u64);
+		OP_TO_STRING_CASE(ds_max_rtn_u64);
+		OP_TO_STRING_CASE(ds_and_rtn_b64);
+		OP_TO_STRING_CASE(ds_or_rtn_b64);
+		OP_TO_STRING_CASE(ds_xor_rtn_b64);
+		OP_TO_STRING_CASE(ds_mskor_rtn_b64);
+		OP_TO_STRING_CASE(ds_wrxchg_rtn_b64);
+		OP_TO_STRING_CASE(ds_wrxchg2_rtn_b64);
+		OP_TO_STRING_CASE(ds_wrxchg2st64_rtn_b64);
+		OP_TO_STRING_CASE(ds_cmpst_rtn_b64);
+		OP_TO_STRING_CASE(ds_cmpst_rtn_f64);
+		OP_TO_STRING_CASE(ds_min_rtn_f64);
+		OP_TO_STRING_CASE(ds_max_rtn_f64);
+		OP_TO_STRING_CASE(ds_read_b64);
+		OP_TO_STRING_CASE(ds_read2_b64);
+		OP_TO_STRING_CASE(ds_read2st64_b64);
+		OP_TO_STRING_CASE(ds_condxchg32_rtn_b64);
+		OP_TO_STRING_CASE(ds_write_b8_d16_hi);
+		OP_TO_STRING_CASE(ds_write_b16_d16_hi);
+		OP_TO_STRING_CASE(ds_read_u8_d16);
+		OP_TO_STRING_CASE(ds_read_u8_d16_hi);
+		OP_TO_STRING_CASE(ds_read_i8_d16);
+		OP_TO_STRING_CASE(ds_read_i8_d16_hi);
+		OP_TO_STRING_CASE(ds_read_u16_d16);
+		OP_TO_STRING_CASE(ds_read_u16_d16_hi);
+		OP_TO_STRING_CASE(ds_write_addtid_b32);
+		OP_TO_STRING_CASE(ds_read_addtid_b32);
+		OP_TO_STRING_CASE(ds_permute_b32);
+		OP_TO_STRING_CASE(ds_bpermute_b32);
+		OP_TO_STRING_CASE(ds_write_b96);
+		OP_TO_STRING_CASE(ds_write_b128);
+		OP_TO_STRING_CASE(ds_read_b96);
+		OP_TO_STRING_CASE(ds_read_b128);
+	default:
+		return "UNKNOWN - LDS / GDS";
+	}
+
+#undef OP_TO_STRING_CASE
+}
+
 enum class EGLOBALOps : uint16
 {
 	global_load_ubyte = 8,
@@ -2425,9 +2689,8 @@ void PrintVOPC(const FInstVOPC& Inst)
 
 void PrintLDSGDS(const FInstLDSGDS& Inst)
 {
-	//ELDSGDSOps Op = (ELDSGDSOps)Inst.OP;
-	//UE_LOG(LogTemp, Warning, TEXT("%hs"), ToString(Op));
-	UE_LOG(LogTemp, Warning, TEXT("LDSGDS - todo"));
+	ELDSGDSOps Op = (ELDSGDSOps)Inst.OP;
+	UE_LOG(LogTemp, Warning, TEXT("%hs"), ToString(Op));
 }
 
 void PrintMUBUF(const FInstMUBUF& Inst)
diff --git a/Engine/Source/Developer/ShaderCompilerCommon/Private/ISAParser/RDNA_Common.h b/Engine/Source/Developer/ShaderCompilerCommon/Private/ISAParser/RDNA_Common.h
--- a/Engine/Source/Developer/ShaderCompilerCommon/Private/ISAParser/RDNA_Common.h
+++ b/Engine/Source/Developer/ShaderCompilerCommon/Private/ISAParser/RDNA_Common.h
@@ -207,7 +207,7 @@ struct FInstLDSGDS
 	uint32 UNUSED : 1;
 	uint32 GDS : 1;
 	uint32 OP : 8;
-	uint32 ID : 8;
+	uint32 ID : 6;
 	uint32 ADDR : 8;
 	uint32 DATA0 : 8;
 	uint32 DATA1 : 8;
@@ -392,11 +392,9 @@ bool IsSOPK(const uint32* Code, uint32& AdvanceAmount)
 	if (bMatch)
 	{
 		AdvanceAmount = GetAdvanceAmount(EInstructionType::SOPK);
-		if (SOPK.SIMM16 == Operand_Literal)
-		{
-			// Skip 32bit literal next in the instruction stream
-			AdvanceAmount++;
-		}
+
+		// SOPK do not ever have a trailing literal as per the spec
+		// "Instructions in this format may not use a 32-bit literal constant which occurs immediately after the instruction."
 	}
 	return bMatch;
 }
@@ -478,7 +476,9 @@ bool IsVOP3(const uint32* Code, uint32& AdvanceAmount)
 	if (bMatch)
 	{
 		AdvanceAmount = GetAdvanceAmount(EInstructionType::VOP3);
-		if (VOP3.SRC0 == Operand_DPP8 || VOP3.SRC0 == Operand_DPP8FI || VOP3.SRC0 == Operand_DPP16 || VOP3.SRC0 == Operand_SDWA || VOP3.SRC0 == Operand_Literal || VOP3.SRC2 == Operand_Literal)
+		if (VOP3.SRC0 == Operand_DPP8 || VOP3.SRC0 == Operand_DPP8FI || VOP3.SRC0 == Operand_DPP16 || VOP3.SRC0 == Operand_SDWA || VOP3.SRC0 == Operand_Literal ||
+			VOP3.SRC1 == Operand_DPP8 || VOP3.SRC1 == Operand_DPP8FI || VOP3.SRC1 == Operand_DPP16 || VOP3.SRC1 == Operand_SDWA || VOP3.SRC1 == Operand_Literal ||
+			VOP3.SRC2 == Operand_DPP8 || VOP3.SRC2 == Operand_DPP8FI || VOP3.SRC2 == Operand_DPP16 || VOP3.SRC2 == Operand_SDWA || VOP3.SRC2 == Operand_Literal)
 		{
 			AdvanceAmount++;
 		}
@@ -626,4 +626,36 @@ EInstructionType DecodeInstructionType(const uint32* Code, uint32& AdvanceAmount
 	if (IsSOP2(Code, AdvanceAmount)) return EInstructionType::SOP2;
 	if (IsVOP2(Code, AdvanceAmount)) return EInstructionType::VOP2;
 	return EInstructionType::UNKNOWN;
+}
+
+const char* ToString(EInstructionType Inst)
+{
+#define OP_TO_STRING_CASE(x) case EInstructionType::x: return #x
+
+	switch (Inst)
+	{
+		OP_TO_STRING_CASE(SOP2);
+		OP_TO_STRING_CASE(SOPK);
+		OP_TO_STRING_CASE(SOP1);
+		OP_TO_STRING_CASE(SOPC);
+		OP_TO_STRING_CASE(SOPP);
+		OP_TO_STRING_CASE(SMEM);
+		OP_TO_STRING_CASE(VOP2);
+		OP_TO_STRING_CASE(VOP1);
+		OP_TO_STRING_CASE(VOP3);
+		OP_TO_STRING_CASE(VOPC);
+		OP_TO_STRING_CASE(VOP3P);
+		OP_TO_STRING_CASE(VINTERP);
+		OP_TO_STRING_CASE(LDSGDS);
+		OP_TO_STRING_CASE(MUBUF);
+		OP_TO_STRING_CASE(MTBUF);
+		OP_TO_STRING_CASE(MIMG);
+		OP_TO_STRING_CASE(EXPORT);
+		OP_TO_STRING_CASE(FSG);
+
+	default:
+		return "UNKNOWN";
+	}
+
+#undef OP_TO_STRING_CASE
 }
\ No newline at end of file
