[shaders] fix minifier test inputs to include simd padding

diff --git a/Engine/Source/Developer/ShaderCompilerCommon/Private/ShaderMinifier.cpp b/Engine/Source/Developer/ShaderCompilerCommon/Private/ShaderMinifier.cpp
--- a/Engine/Source/Developer/ShaderCompilerCommon/Private/ShaderMinifier.cpp
+++ b/Engine/Source/Developer/ShaderCompilerCommon/Private/ShaderMinifier.cpp
@@ -650,13 +650,13 @@ FShaderSource::FViewType ExtractNextIdentifier(FShaderSource::FViewType Source)
 	return Identifier;
 }
 
-static FParsedShader ParseShader(FShaderSource::FViewType InSource, FDiagnostics& Output)
+static FParsedShader ParseShader(const FShaderSource& InSource, FDiagnostics& Output)
 {
 	FParsedShader Result;
 
-	Result.Source = InSource;
+	Result.Source = InSource.GetView();
 
-	FShaderSource::FViewType	Source = InSource;
+	FShaderSource::FViewType	Source = InSource.GetView();
 	FCodeBlockArray			PendingBlocks;
 	TArray<FCodeChunk>		Chunks;
 
@@ -682,7 +682,7 @@ static FParsedShader ParseShader(FShaderSource::FViewType InSource, FDiagnostics
 		FDiagnosticMessage Diagnostic;
 
 		Diagnostic.Message = FString(Message);
-		Diagnostic.Offset = int32(Source.GetData() - InSource.GetData());
+		Diagnostic.Offset = int32(Source.GetData() - InSource.GetView().GetData());
 		// Diagnostic.Line = ...; // TODO
 		// Diagnostic.Column = ...; // TODO
 
@@ -2076,7 +2076,7 @@ FMinifiedShader Minify(const FShaderSource& PreprocessedShader, TConstArrayView<
 {
 	FMinifiedShader Result;
 
-	FParsedShader Parsed = ParseShader(PreprocessedShader.GetView(), Result.Diagnostics);
+	FParsedShader Parsed = ParseShader(PreprocessedShader, Result.Diagnostics);
 
 	if (!Parsed.Chunks.IsEmpty())
 	{
@@ -2100,7 +2100,7 @@ IMPLEMENT_SIMPLE_AUTOMATION_TEST(FShaderMinifierParserTest, "System.Shaders.Shad
 namespace UE::ShaderMinifier
 {
 // Convenience wrapper for tests where we don't care about diagnostic messages
-static FParsedShader ParseShader(FShaderSource::FViewType InSource)
+static FParsedShader ParseShader(const FShaderSource& InSource)
 {
 	FDiagnostics Diagnostics;
 	return ParseShader(InSource, Diagnostics);
@@ -2111,20 +2111,28 @@ bool FShaderMinifierParserTest::RunTest(const FString& Parameters)
 {
 	using namespace UE::ShaderMinifier;
 
-	TestEqual(TEXT("SkipSpace"), 
-		FString(SkipSpace(SHADER_SOURCE_LITERAL("  \n\r\f \tHello"))), 
-		FString(SHADER_SOURCE_LITERAL("Hello")));
-
-	TestEqual(TEXT("SkipUntilStr (found)"), 
-		FString(SkipUntilStr(SHADER_SOURCE_LITERAL("Hello World"), SHADER_SOURCE_LITERAL("World"))),
-		FString(SHADER_SOURCE_LITERAL("World")));
-
-	TestEqual(TEXT("SkipUntilStr (not found)"),
-		FString(SkipUntilStr(SHADER_SOURCE_LITERAL("Hello World"), SHADER_SOURCE_LITERAL("Blah"))),
-		FString());
+	{
+		FShaderSource S(SHADER_SOURCE_LITERAL("  \n\r\f \tHello"));
+		TestEqual(TEXT("SkipSpace"),
+			FString(SkipSpace(S.GetView())),
+			FString(SHADER_SOURCE_LITERAL("Hello")));
+	}
+	{
+		FShaderSource S(SHADER_SOURCE_LITERAL("Hello World"));
+		TestEqual(TEXT("SkipUntilStr (found)"),
+			FString(SkipUntilStr(S.GetView(), SHADER_SOURCE_LITERAL("World"))),
+			FString(SHADER_SOURCE_LITERAL("World")));
+	}
+	{
+		FShaderSource S(SHADER_SOURCE_LITERAL("Hello World"));
+		TestEqual(TEXT("SkipUntilStr (not found)"),
+			FString(SkipUntilStr(S.GetView(), SHADER_SOURCE_LITERAL("Blah"))),
+			FString());
+	}
 
 	{
-		auto P = ParseShader(SHADER_SOURCE_LITERAL("static const struct { int Blah; } Foo = { 123; };"));
+		FShaderSource S(SHADER_SOURCE_LITERAL("static const struct { int Blah; } Foo = { 123; };"));
+		auto P = ParseShader(S);
 		TestEqual(TEXT("Anonymous struct variable with initializer, total chunks"), P.Chunks.Num(), 1);
 		if (P.Chunks.Num() == 1)
 		{
@@ -2133,7 +2141,8 @@ bool FShaderMinifierParserTest::RunTest(const FString& Parameters)
 	}
 
 	{
-		auto P = ParseShader(SHADER_SOURCE_LITERAL("float4 PSMain() : SV_Target { return float4(1,0,0,1); };"));
+		FShaderSource S(SHADER_SOURCE_LITERAL("float4 PSMain() : SV_Target { return float4(1,0,0,1); };"));
+		auto P = ParseShader(S);
 		TestEqual(TEXT("Pixel shader entry point, total chunks"), P.Chunks.Num(), 1);
 		if (P.Chunks.Num() == 1)
 		{
@@ -2143,7 +2152,8 @@ bool FShaderMinifierParserTest::RunTest(const FString& Parameters)
 
 	{
 		TArray<FShaderSource::FViewType> R;
-		ExtractIdentifiers(SHADER_SOURCE_LITERAL("Hello[World]; Foo[0];\n"), R);
+		FShaderSource S(SHADER_SOURCE_LITERAL("Hello[World]; Foo[0];\n"));
+		ExtractIdentifiers(S.GetView(), R);
 		if (TestEqual(TEXT("ExtractIdentifiers1: Num"), R.Num(), 3))
 		{
 			TestEqual(TEXT("ExtractIdentifiers1: R[0]"), FString(R[0]), SHADER_SOURCE_LITERAL("Hello"));
@@ -2154,7 +2164,8 @@ bool FShaderMinifierParserTest::RunTest(const FString& Parameters)
 
 	{
 		TArray<FShaderSource::FViewType> R;
-		ExtractIdentifiers(SHADER_SOURCE_LITERAL("#line 0\nStructuredBuffer<uint4> Blah : register(t0, space123);#line 1\n#pragma foo\n"), R);
+		FShaderSource S(SHADER_SOURCE_LITERAL("#line 0\nStructuredBuffer<uint4> Blah : register(t0, space123);#line 1\n#pragma foo\n"));
+		ExtractIdentifiers(S.GetView(), R);
 		if (TestEqual(TEXT("ExtractIdentifiers2: Num"), R.Num(), 6))
 		{
 			TestEqual(TEXT("ExtractIdentifiers2: R[0]"), FString(R[0]), SHADER_SOURCE_LITERAL("StructuredBuffer"));
@@ -2167,7 +2178,8 @@ bool FShaderMinifierParserTest::RunTest(const FString& Parameters)
 	}
 
 	{
-		auto P = ParseShader(SHADER_SOURCE_LITERAL("StructuredBuffer<uint4> Blah : register(t0, space123);"));
+		FShaderSource S(SHADER_SOURCE_LITERAL("StructuredBuffer<uint4> Blah : register(t0, space123);"));
+		auto P = ParseShader(S);
 		if (TestEqual(TEXT("ParseShader: structured buffer: num chunks"), P.Chunks.Num(), 1))
 		{
 			TestEqual(TEXT("ParseShader: structured buffer: chunk"), P.Chunks[0].Type, ECodeChunkType::Variable);
@@ -2175,7 +2187,8 @@ bool FShaderMinifierParserTest::RunTest(const FString& Parameters)
 	}
 
 	{
-		auto P = ParseShader(SHADER_SOURCE_LITERAL("const float Foo = 123.45f;"));
+		FShaderSource S(SHADER_SOURCE_LITERAL("const float Foo = 123.45f;"));
+		auto P = ParseShader(S);
 		if (TestEqual(TEXT("ParseShader: const float with initializer: num chunks"), P.Chunks.Num(), 1))
 		{
 			TestEqual(TEXT("ParseShader: const float with initializer: chunk type"), P.Chunks[0].Type, ECodeChunkType::Variable);
@@ -2183,7 +2196,8 @@ bool FShaderMinifierParserTest::RunTest(const FString& Parameters)
 	}
 
 	{
-		auto P = ParseShader(SHADER_SOURCE_LITERAL("struct Blah { int A; };"));
+		FShaderSource S(SHADER_SOURCE_LITERAL("struct Blah { int A; };"));
+		auto P = ParseShader(S);
 		if (TestEqual(TEXT("ParseShader: struct: num chunks"), P.Chunks.Num(), 1))
 		{
 			TestEqual(TEXT("ParseShader: struct: chunk type"), P.Chunks[0].Type, ECodeChunkType::Struct);
@@ -2191,7 +2205,8 @@ bool FShaderMinifierParserTest::RunTest(const FString& Parameters)
 	}
 
 	{
-		auto P = ParseShader(SHADER_SOURCE_LITERAL("struct Foo { int FooA; }; struct Bar : Foo { int BarA; };"));
+		FShaderSource S(SHADER_SOURCE_LITERAL("struct Foo { int FooA; }; struct Bar : Foo { int BarA; };"));
+		auto P = ParseShader(S);
 		if (TestEqual(TEXT("ParseShader: inherited struct: num chunks"), P.Chunks.Num(), 2))
 		{
 			TestEqual(TEXT("ParseShader: inherited struct: chunk 0 type"), P.Chunks[0].Type, ECodeChunkType::Struct);
@@ -2200,7 +2215,8 @@ bool FShaderMinifierParserTest::RunTest(const FString& Parameters)
 	}
 
 	{
-		auto P = ParseShader(SHADER_SOURCE_LITERAL("[numthreads(8,8,1)] void Main() {};"));
+		FShaderSource S(SHADER_SOURCE_LITERAL("[numthreads(8,8,1)] void Main() {};"));
+		auto P = ParseShader(S);
 		if (TestEqual(TEXT("ParseShader: compute shader entry point: num chunks"), P.Chunks.Num(), 1))
 		{
 			TestEqual(TEXT("ParseShader: compute shader entry point: chunk type"), P.Chunks[0].Type, ECodeChunkType::Function);
@@ -2213,14 +2229,16 @@ bool FShaderMinifierParserTest::RunTest(const FString& Parameters)
 	}
 
 	{
-		auto P = ParseShader(SHADER_SOURCE_LITERAL("Texture2D Blah : register(t0);"));
+		FShaderSource S(SHADER_SOURCE_LITERAL("Texture2D Blah : register(t0);"));
+		auto P = ParseShader(S);
 		if (TestEqual(TEXT("ParseShader: texture with register: num chunks"), P.Chunks.Num(), 1))
 		{
 			TestEqual(TEXT("ParseShader: texture with register: chunk type"), P.Chunks[0].Type, ECodeChunkType::Variable);
 		}
 	}
 	{
-		auto P = ParseShader(SHADER_SOURCE_LITERAL("Texture2D Blah;"));
+		FShaderSource S(SHADER_SOURCE_LITERAL("Texture2D Blah;"));
+		auto P = ParseShader(S);
 		if (TestEqual(TEXT("ParseShader: texture: num chunks"), P.Chunks.Num(), 1))
 		{
 			TestEqual(TEXT("ParseShader: texture: chunk type"), P.Chunks[0].Type, ECodeChunkType::Variable);
@@ -2228,7 +2246,8 @@ bool FShaderMinifierParserTest::RunTest(const FString& Parameters)
 	}
 
 	{
-		auto P = ParseShader(SHADER_SOURCE_LITERAL("SamplerState Blah : register(s0, space123);"));
+		FShaderSource S(SHADER_SOURCE_LITERAL("SamplerState Blah : register(s0, space123);"));
+		auto P = ParseShader(S);
 		if (TestEqual(TEXT("ParseShader: sampler state with register: num chunks"), P.Chunks.Num(), 1))
 		{
 			TestEqual(TEXT("ParseShader: sampler state with register: chunk type"), P.Chunks[0].Type, ECodeChunkType::Variable);
@@ -2238,13 +2257,15 @@ bool FShaderMinifierParserTest::RunTest(const FString& Parameters)
 #if 0
 	{
 		// TODO: handle function forward declarations
-		auto P = ParseShader(SHADER_SOURCE_LITERAL("Foo Fun(int a);"));
+		FShaderSource S(SHADER_SOURCE_LITERAL("Foo Fun(int a);"));
+		auto P = ParseShader(S);
 		TestEqual(TEXT("ParseShader: function forward declaration"), P.Chunks[0].Type, ECodeChunkType::FunctionDecl);
 	}
 #endif
 
 	{
-		auto P = ParseShader(SHADER_SOURCE_LITERAL("void Fun(int a) {};"));
+		FShaderSource S(SHADER_SOURCE_LITERAL("void Fun(int a) {};"));
+		auto P = ParseShader(S);
 		if (TestEqual(TEXT("ParseShader: function with trailing semicolon: num chunks"), P.Chunks.Num(), 1))
 		{
 			TestEqual(TEXT("ParseShader: function with trailing semicolon: chunk type"), P.Chunks[0].Type, ECodeChunkType::Function);
@@ -2252,7 +2273,8 @@ bool FShaderMinifierParserTest::RunTest(const FString& Parameters)
 	}
 
 	{
-		auto P = ParseShader(SHADER_SOURCE_LITERAL("void Fun(int a) {}"));
+		FShaderSource S(SHADER_SOURCE_LITERAL("void Fun(int a) {}"));
+		auto P = ParseShader(S);
 		if (TestEqual(TEXT("ParseShader: function: num chunks"), P.Chunks.Num(), 1))
 		{
 			TestEqual(TEXT("ParseShader: function: chunk type"), P.Chunks[0].Type, ECodeChunkType::Function);
@@ -2260,7 +2282,8 @@ bool FShaderMinifierParserTest::RunTest(const FString& Parameters)
 	}
 
 	{
-		auto P = ParseShader(SHADER_SOURCE_LITERAL("cbuffer Foo {blah} SamplerState S;"));
+		FShaderSource S(SHADER_SOURCE_LITERAL("cbuffer Foo {blah} SamplerState S;"));
+		auto P = ParseShader(S);
 
 		if (TestEqual(TEXT("ParseShader: cbuffer and sampler state: num chunks"), P.Chunks.Num(), 2))
 		{
@@ -2270,7 +2293,8 @@ bool FShaderMinifierParserTest::RunTest(const FString& Parameters)
 	}
 
 	{
-		auto P = ParseShader(SHADER_SOURCE_LITERAL("struct Foo { int a; };"));
+		FShaderSource S(SHADER_SOURCE_LITERAL("struct Foo { int a; };"));
+		auto P = ParseShader(S);
 		if (TestEqual(TEXT("ParseShader: struct: num chunks"), P.Chunks.Num(), 1))
 		{
 			TestEqual(TEXT("ParseShader: struct: chunk type"), P.Chunks[0].Type, ECodeChunkType::Struct);
@@ -2278,7 +2302,8 @@ bool FShaderMinifierParserTest::RunTest(const FString& Parameters)
 	}
 
 	{
-		auto P = ParseShader(SHADER_SOURCE_LITERAL("struct { int a; } Foo = { 123; };"));
+		FShaderSource S(SHADER_SOURCE_LITERAL("struct { int a; } Foo = { 123; };"));
+		auto P = ParseShader(S);
 		if (TestEqual(TEXT("ParseShader: anonymous struct with variable and initializer: num chunks"), P.Chunks.Num(), 1))
 		{
 			TestEqual(TEXT("ParseShader: anonymous struct with variable and initializer: chunk type [0]"), P.Chunks[0].Type, ECodeChunkType::Variable);
@@ -2288,14 +2313,18 @@ bool FShaderMinifierParserTest::RunTest(const FString& Parameters)
 #if 0
 	{
 		// TODO: handle struct forward declarations
-		auto P = ParseShader(SHADER_SOURCE_LITERAL("struct Foo;"));
+		FShaderSource S(SHADER_SOURCE_LITERAL("struct Foo;"));
+		auto P = ParseShader(S);
 	}
 #endif
 
 	{
-		auto P = ParseShader(
-			SHADER_SOURCE_LITERAL("cbuffer MyBuffer : register(b3)")
-			SHADER_SOURCE_LITERAL("{ float4 Element1 : packoffset(c0); float1 Element2 : packoffset(c1); float1 Element3 : packoffset(c1.y); }"));
+		FShaderSource S(
+			SHADER_SOURCE_LITERAL(
+				"cbuffer MyBuffer : register(b3)"
+				"{ float4 Element1 : packoffset(c0); float1 Element2 : packoffset(c1); float1 Element3 : packoffset(c1.y); }")
+			);
+		auto P = ParseShader(S);
 		if (TestEqual(TEXT("ParseShader: cbuffer with packoffset: num chunks"), P.Chunks.Num(), 1))
 		{
 			TestEqual(TEXT("ParseShader: cbuffer with packoffset: chunk type"), P.Chunks[0].Type, ECodeChunkType::CBuffer);
@@ -2303,7 +2332,8 @@ bool FShaderMinifierParserTest::RunTest(const FString& Parameters)
 	}
 
 	{
-		auto P = ParseShader(SHADER_SOURCE_LITERAL("static const struct { float4 Param; } Foo;"));
+		FShaderSource S(SHADER_SOURCE_LITERAL("static const struct { float4 Param; } Foo;"));
+		auto P = ParseShader(S);
 		if (TestEqual(TEXT("ParseShader: static const anonymous struct with variable: num chunks"), P.Chunks.Num(), 1))
 		{
 			TestEqual(TEXT("ParseShader: static const anonymous struct with variable: chunk type"), P.Chunks[0].Type, ECodeChunkType::Variable);
@@ -2311,7 +2341,8 @@ bool FShaderMinifierParserTest::RunTest(const FString& Parameters)
 	}
 
 	{
-		auto P = ParseShader(SHADER_SOURCE_LITERAL("static const struct { float4 Param; } Foo = { FooCB_Param; };"));
+		FShaderSource S(SHADER_SOURCE_LITERAL("static const struct { float4 Param; } Foo = { FooCB_Param; };"));
+		auto P = ParseShader(S);
 		if (TestEqual(TEXT("ParseShader: static const anonymous struct with variable and initializer: num chunks"), P.Chunks.Num(), 1))
 		{
 			TestEqual(TEXT("ParseShader: static const anonymous struct with variable and initializer: chunk type"), P.Chunks[0].Type, ECodeChunkType::Variable);
@@ -2319,7 +2350,8 @@ bool FShaderMinifierParserTest::RunTest(const FString& Parameters)
 	}
 
 	{
-		auto P = ParseShader(SHADER_SOURCE_LITERAL("template <typename T> float Fun(T x) { return (float)x; }"));
+		FShaderSource S(SHADER_SOURCE_LITERAL("template <typename T> float Fun(T x) { return (float)x; }"));
+		auto P = ParseShader(S);
 		if (TestEqual(TEXT("ParseShader: template function: num chunks"), P.Chunks.Num(), 1))
 		{
 			TestEqual(TEXT("ParseShader: template function: chunk type"), P.Chunks[0].Type, ECodeChunkType::Function);
@@ -2327,7 +2359,8 @@ bool FShaderMinifierParserTest::RunTest(const FString& Parameters)
 	}
 
 	{
-		auto P = ParseShader(SHADER_SOURCE_LITERAL("enum EFoo { A, B = 123 };"));
+		FShaderSource S(SHADER_SOURCE_LITERAL("enum EFoo { A, B = 123 };"));
+		auto P = ParseShader(S);
 		if (TestEqual(TEXT("ParseShader: enum: num chunks"), P.Chunks.Num(), 1))
 		{
 			TestEqual(TEXT("ParseShader: enum: chunk type"), P.Chunks[0].Type, ECodeChunkType::Enum);
@@ -2335,7 +2368,8 @@ bool FShaderMinifierParserTest::RunTest(const FString& Parameters)
 	}
 
 	{
-		auto P = ParseShader(SHADER_SOURCE_LITERAL("enum class EFoo { A, B };"));
+		FShaderSource S(SHADER_SOURCE_LITERAL("enum class EFoo { A, B };"));
+		auto P = ParseShader(S);
 		if (TestEqual(TEXT("ParseShader: enum class: num chunks"), P.Chunks.Num(), 1))
 		{
 			TestEqual(TEXT("ParseShader: enum class: chunk type"), P.Chunks[0].Type, ECodeChunkType::Enum);
@@ -2343,7 +2377,8 @@ bool FShaderMinifierParserTest::RunTest(const FString& Parameters)
 	}
 
 	{
-		auto P = ParseShader(SHADER_SOURCE_LITERAL("#define Foo 123"));
+		FShaderSource S(SHADER_SOURCE_LITERAL("#define Foo 123"));
+		auto P = ParseShader(S);
 		if (TestEqual(TEXT("ParseShader: define: num chunks"), P.Chunks.Num(), 1))
 		{
 			TestEqual(TEXT("ParseShader: define: chunk type"), P.Chunks[0].Type, ECodeChunkType::Define);
@@ -2351,7 +2386,8 @@ bool FShaderMinifierParserTest::RunTest(const FString& Parameters)
 	}
 
 	{
-		auto P = ParseShader(SHADER_SOURCE_LITERAL("#pragma Foo"));
+		FShaderSource S(SHADER_SOURCE_LITERAL("#pragma Foo"));
+		auto P = ParseShader(S);
 		if (TestEqual(TEXT("ParseShader: pragma: num chunks"), P.Chunks.Num(), 1))
 		{
 			TestEqual(TEXT("ParseShader: pragma: chunk type"), P.Chunks[0].Type, ECodeChunkType::Pragma);
@@ -2359,7 +2395,8 @@ bool FShaderMinifierParserTest::RunTest(const FString& Parameters)
 	}
 
 	{
-		auto P = ParseShader(SHADER_SOURCE_LITERAL("ConstantBuffer<Foo> CB : register ( b123, space456);"));
+		FShaderSource S(SHADER_SOURCE_LITERAL("ConstantBuffer<Foo> CB : register ( b123, space456);"));
+		auto P = ParseShader(S);
 		if (TestEqual(TEXT("ParseShader: ConstantBuffer<Foo>: num chunks"), P.Chunks.Num(), 1))
 		{
 			TestEqual(TEXT("ParseShader: ConstantBuffer<Foo>: chunk type"), P.Chunks[0].Type, ECodeChunkType::Variable);
@@ -2367,7 +2404,8 @@ bool FShaderMinifierParserTest::RunTest(const FString& Parameters)
 	}
 
 	{
-		auto P = ParseShader(SHADER_SOURCE_LITERAL("namespace NS1 { void Fun() {}; } namespace NS2 { void Fun() {}; }"));
+		FShaderSource S(SHADER_SOURCE_LITERAL("namespace NS1 { void Fun() {}; } namespace NS2 { void Fun() {}; }"));
+		auto P = ParseShader(S);
 		if (TestEqual(TEXT("ParseShader: namespaces: num chunks"), P.Chunks.Num(), 2)
 			&& TestEqual(TEXT("ParseShader: namespaces: num namespaces"), P.Namespaces.Num(), 2))
 		{
@@ -2377,7 +2415,8 @@ bool FShaderMinifierParserTest::RunTest(const FString& Parameters)
 	}
 
 	{
-		auto P = ParseShader(SHADER_SOURCE_LITERAL("template< typename T > TMyStruct<T> operator + ( TMyStruct<T> A, T B ) { /*...*/ }"));
+		FShaderSource S(SHADER_SOURCE_LITERAL("template< typename T > TMyStruct<T> operator + ( TMyStruct<T> A, T B ) { /*...*/ }"));
+		auto P = ParseShader(S);
 		if (TestEqual(TEXT("ParseShader: operators: num chunks"), P.Chunks.Num(), 1))
 		{
 			TestEqual(TEXT("ParseShader: operators: chunk type"), P.Chunks[0].Type, ECodeChunkType::Operator);
@@ -2396,7 +2435,8 @@ bool FShaderMinifierParserTest::RunTest(const FString& Parameters)
 	}
 
 	{
-		auto P = ParseShader(SHADER_SOURCE_LITERAL("typedef Bar Foo;"));
+		FShaderSource S(SHADER_SOURCE_LITERAL("typedef Bar Foo;"));
+		auto P = ParseShader(S);
 		if (TestEqual(TEXT("ParseShader: standard typedef: num chunks"), P.Chunks.Num(), 1))
 		{
 			TestEqual(TEXT("ParseShader: standard typedef: chunk type"), P.Chunks[0].Type, ECodeChunkType::Typedef);
@@ -2404,7 +2444,8 @@ bool FShaderMinifierParserTest::RunTest(const FString& Parameters)
 	}
 
 	{
-		auto P = ParseShader(SHADER_SOURCE_LITERAL("typedef Bar Foo; static const Foo = Bar(0);"));
+		FShaderSource S(SHADER_SOURCE_LITERAL("typedef Bar Foo; static const Foo = Bar(0);"));
+		auto P = ParseShader(S);
 		if (TestEqual(TEXT("ParseShader: standard typedef: num chunks"), P.Chunks.Num(), 2))
 		{
 			TestEqual(TEXT("ParseShader: standard typedef: chunk type"), P.Chunks[0].Type, ECodeChunkType::Typedef);
@@ -2423,7 +2464,7 @@ bool FShaderMinifierTest::RunTest(const FString& Parameters)
 {
 	using namespace UE::ShaderMinifier;
 
-	FShaderSource::FViewType TestShaderCode = 
+	FShaderSource TestShaderCode(
 		SHADER_SOURCE_LITERAL(R"(// dxc /T cs_6_6 /E MainCS MinifierTest.hlsl 
 struct FFoo
 {
@@ -2563,7 +2604,7 @@ void MainCS()
 	float D = TypedefUsedBuffer[ENUM_USED_PART_1].Foo;
 	OutputBuffer[0] = A + B + D;
 }
-)");
+)"));
 
 	auto ChunkPresent = [](const FParsedShader& Parsed, FShaderSource::FViewType Name)
 	{
@@ -2584,7 +2625,7 @@ void MainCS()
 
 	{
 		FDiagnostics Diagnostics;
-		FShaderSource::FStringType Minified = MinifyShader(Parsed, SHADER_SOURCE_LITERAL("EmptyFunction"), EMinifyShaderFlags::None, Diagnostics);
+		FShaderSource Minified = FShaderSource(MinifyShader(Parsed, SHADER_SOURCE_LITERAL("EmptyFunction"), EMinifyShaderFlags::None, Diagnostics));
 		FParsedShader MinifiedParsed = ParseShader(Minified);
 		if (TestEqual(TEXT("MinifyShader: EmptyFunction: num chunks"), MinifiedParsed.Chunks.Num(), 3))
 		{
@@ -2596,7 +2637,7 @@ void MainCS()
 
 	{
 		FDiagnostics Diagnostics;
-		FShaderSource::FStringType Minified = MinifyShader(Parsed, SHADER_SOURCE_LITERAL("MainCS"), EMinifyShaderFlags::OutputReasons, Diagnostics);
+		FShaderSource Minified = FShaderSource(MinifyShader(Parsed, SHADER_SOURCE_LITERAL("MainCS"), EMinifyShaderFlags::OutputReasons, Diagnostics));
 		FParsedShader MinifiedParsed = ParseShader(Minified);
 
 		// Expect true:
