[shaders] shaderconductorcontext improvements:
- add support to shaderconductor to load shader source from a FAnsiStringView
- remove internal string conversion helpers in favour of using FAnsiString
- don't make copies of input strings if they are provided in ANSI format

diff --git a/Engine/Source/Developer/ShaderCompilerCommon/Private/ShaderConductorContext.cpp b/Engine/Source/Developer/ShaderCompilerCommon/Private/ShaderConductorContext.cpp
--- a/Engine/Source/Developer/ShaderCompilerCommon/Private/ShaderConductorContext.cpp
+++ b/Engine/Source/Developer/ShaderCompilerCommon/Private/ShaderConductorContext.cpp
@@ -3,6 +3,7 @@
 #include "ShaderConductorContext.h"
 #include "HAL/ExceptionHandling.h"
 #include "ShaderCompilerDefinitions.h"
+#include "Containers/AnsiString.h"
 
 #if PLATFORM_MAC || PLATFORM_WINDOWS || PLATFORM_LINUX
 THIRD_PARTY_INCLUDES_START
@@ -43,53 +44,13 @@ namespace CrossCompiler
 		OutResultDesc = ShaderConductor::Compiler::ConvertBinary(InBinaryDesc, InSourceDesc, InOptions, InTargetDesc);
 	}
 
-	// Converts the byte array 'InString' (without null terminator) to the output ANSI string 'OutString' (with appended null terminator).
-	static void ConvertByteArrayToAnsiString(const ANSICHAR* InString, uint32 InStringLength, TArray<ANSICHAR>& OutString)
-	{
-		// 'FCStringAnsi::Strncpy()' will put a '\0' character at the end
-		OutString.SetNum(InStringLength + 1);
-		FCStringAnsi::Strncpy(OutString.GetData(), InString, OutString.Num());
-	}
-
-	// Converts the FString to the output ANSI string 'OutString'.
-	static void ConvertFStringToAnsiString(const FString& InString, TArray<ANSICHAR>& OutString)
-	{
-		ConvertByteArrayToAnsiString(TCHAR_TO_ANSI(*InString), InString.Len(), OutString);
-	}
-
-	static void ConvertFStringViewToAnsiString(const FStringView& InString, TArray<ANSICHAR>& OutString)
-	{
-		ConvertByteArrayToAnsiString(TCHAR_TO_ANSI(InString.GetData()), InString.Len(), OutString);
-	}
-
-	// Copies the NULL-terminated string 'InString' to 'OutString'. Also copies the '\0' character at the end.
-	static void CopyAnsiString(const ANSICHAR* InString, TArray<ANSICHAR>& OutString)
-	{
-		// 'InString' is NULL-terminated, so we can use 'FCStringAnsi::Strlen()'
-		if (InString != nullptr)
-		{
-			ConvertByteArrayToAnsiString(InString, FCStringAnsi::Strlen(InString), OutString);
-		}
-	}
-
-	// Converts the specified ShaderConductor blob to FString.
-	static bool ConvertByteArrayToFString(const void* InData, uint32 InSize, FString& OutString)
-	{
-		if (InData != nullptr && InSize > 0)
-		{
-			FUTF8ToTCHAR UTF8Converter(reinterpret_cast<const ANSICHAR*>(InData), InSize);
-			OutString = FString(UTF8Converter.Length(), UTF8Converter.Get());
-			return true;
-		}
-		return false;
-	}
-
 	// Converts the specified ShaderConductor blob to FString.
 	static bool ConvertScBlobToFString(ShaderConductor::Blob* Blob, FString& OutString)
 	{
 		if (Blob && Blob->Size() > 0)
 		{
-			return ConvertByteArrayToFString(Blob->Data(), Blob->Size(), OutString);
+			OutString = StringCast<TCHAR>(reinterpret_cast<const ANSICHAR*>(Blob->Data()), Blob->Size());
+			return true;
 		}
 		return false;
 	}
@@ -116,6 +77,32 @@ namespace CrossCompiler
 		return ShaderConductor::ShaderStage::NumShaderStages;
 	}
 
+	struct FOptionalConvertedAnsiString
+	{
+		FOptionalConvertedAnsiString& operator=(FStringView WideView)
+		{
+			Storage.Empty();
+			Storage.Append(WideView);
+			AnsiView = FAnsiStringView(Storage);
+			return *this;
+		}
+
+		FOptionalConvertedAnsiString& operator=(FAnsiStringView InAnsiView)
+		{
+			AnsiView = InAnsiView;
+			return *this;
+		}
+
+		void CopyAnsi(FAnsiStringView InAnsiView)
+		{
+			Storage = InAnsiView;
+			AnsiView = FAnsiStringView(Storage);
+		}
+
+		FAnsiString Storage;
+		FAnsiStringView AnsiView;
+	};
+
 	// Wrapper structure to hold all intermediate buffers for ShaderConductor
 	struct FShaderConductorContext::FShaderConductorIntermediates
 	{
@@ -124,16 +111,16 @@ namespace CrossCompiler
 		{
 		}
 
-		TArray<ANSICHAR> ShaderSource;
-		TArray<ANSICHAR> Filename;
-		TArray<ANSICHAR> EntryPoint;
+		FOptionalConvertedAnsiString ShaderSource;
+		FOptionalConvertedAnsiString Filename;
+		FOptionalConvertedAnsiString EntryPoint;
 		ShaderConductor::ShaderStage Stage;
-		TArray<TPair<TArray<ANSICHAR>, TArray<ANSICHAR>>> Defines;
+		TArray<TPair<FAnsiString, FAnsiString>> Defines;
 		TArray<ShaderConductor::MacroDefine> DefineRefs;
-		TArray<TPair<TArray<ANSICHAR>, TArray<ANSICHAR>>> Flags;
+		TArray<TPair<FAnsiString, FAnsiString>> Flags;
 		TArray<ShaderConductor::MacroDefine> FlagRefs;
-		TArray<ANSICHAR> InternalDxcArgs;
-		TArray<TArray<ANSICHAR>> CustomDxcArgs;
+		FAnsiString InternalDxcArgs;
+		TArray<FAnsiString> CustomDxcArgs;
 		TArray<ANSICHAR const*> CustomDxcArgRefs;
 		TArray<ANSICHAR const*> DxcArgRefs;
 	};
@@ -141,9 +128,9 @@ namespace CrossCompiler
 	static void ConvertScSourceDesc(const FShaderConductorContext::FShaderConductorIntermediates& Intermediates, ShaderConductor::Compiler::SourceDesc& OutSourceDesc)
 	{
 		// Convert descriptor with pointers to the ANSI strings
-		OutSourceDesc.source = Intermediates.ShaderSource.GetData();
-		OutSourceDesc.fileName = Intermediates.Filename.GetData();
-		OutSourceDesc.entryPoint = Intermediates.EntryPoint.GetData();
+		OutSourceDesc.source = Intermediates.ShaderSource.AnsiView.GetData();
+		OutSourceDesc.fileName = Intermediates.Filename.AnsiView.GetData();
+		OutSourceDesc.entryPoint = Intermediates.EntryPoint.AnsiView.GetData();
 		OutSourceDesc.stage = Intermediates.Stage;
 		if (Intermediates.DefineRefs.Num() > 0)
 		{
@@ -229,44 +216,40 @@ namespace CrossCompiler
 	}
 
 	// Converts an array of FString to a C-style array of char* pointers
-	static void ConvertStringArrayToAnsiArray(const TArray<FString>& InPairs, TArray<TArray<ANSICHAR>>& OutPairs, TArray<const char*>& OutPairRefs)
+	static void ConvertStringArrayToAnsiArray(const TArray<FString>& InPairs, TArray<FAnsiString>& OutPairs, TArray<const char*>& OutPairRefs)
 	{
 		// Convert map into an array container
 		TArray<ANSICHAR> Value;
 		for (const FString& Iter : InPairs)
 		{
-			ConvertFStringToAnsiString(Iter, Value);
-			OutPairs.Emplace(MoveTemp(Value));
+			OutPairs.Emplace(Iter);
 		}
 
 		// Store references after all elements have been added to the container so the pointers remain valid
 		OutPairRefs.SetNum(OutPairs.Num());
 		for (int32 Index = 0; Index < OutPairs.Num(); ++Index)
 		{
-			OutPairRefs[Index] = OutPairs[Index].GetData();
+			OutPairRefs[Index] = *OutPairs[Index];
 		}
 	}
 
 	PRAGMA_DISABLE_DEPRECATION_WARNINGS		// FShaderCompilerDefinitions will be made internal in the future, marked deprecated until then
 
 	// Converts a map of string pairs to a C-Style macro defines array
-	static void ConvertDefineMapToMacroDefines(const FShaderCompilerDefinitions& Definitions, TArray<TPair<TArray<ANSICHAR>, TArray<ANSICHAR>>>& OutPairs, TArray<ShaderConductor::MacroDefine>& OutPairRefs)
+	static void ConvertDefineMapToMacroDefines(const FShaderCompilerDefinitions& Definitions, TArray<TPair<FAnsiString, FAnsiString>>& OutPairs, TArray<ShaderConductor::MacroDefine>& OutPairRefs)
 	{
 		// Convert map into an array container
-		TArray<ANSICHAR> Name, Value;
 		for (FShaderCompilerDefinitions::FConstIterator Iter(Definitions); Iter; ++Iter)
 		{
-			ConvertFStringViewToAnsiString(FStringView(Iter.Key()), Name);
-			ConvertFStringViewToAnsiString(FStringView(Iter.Value()), Value);
-			OutPairs.Emplace(MoveTemp(Name), MoveTemp(Value));
+			OutPairs.Emplace(Iter.Key(), Iter.Value());
 		}
 
 		// Store references after all elements have been added to the container so the pointers remain valid
 		OutPairRefs.SetNum(OutPairs.Num());
 		for (int32 Index = 0; Index < OutPairs.Num(); ++Index)
 		{
-			OutPairRefs[Index].name = OutPairs[Index].Key.GetData();
-			OutPairRefs[Index].value = OutPairs[Index].Value.GetData();
+			OutPairRefs[Index].name = *OutPairs[Index].Key;
+			OutPairRefs[Index].value = *OutPairs[Index].Value;
 		}
 	}
 	PRAGMA_ENABLE_DEPRECATION_WARNINGS
@@ -501,8 +484,8 @@ namespace CrossCompiler
 
 		if (!InOptions.SpirvCustomOptimizationPasses.IsEmpty())
 		{
-			ConvertFStringToAnsiString(FString::Printf(TEXT("-Oconfig=%s"), SelectSpirvCustomOptimizationPasses(InOptions.SpirvCustomOptimizationPasses)), Intermediates.InternalDxcArgs);
-			DxcArgRefs.Add(Intermediates.InternalDxcArgs.GetData());
+			Intermediates.InternalDxcArgs = FAnsiString::Printf("-Oconfig=%ls", SelectSpirvCustomOptimizationPasses(InOptions.SpirvCustomOptimizationPasses));
+			DxcArgRefs.Add(*Intermediates.InternalDxcArgs);
 		}
 
 		if (DxcArgRefs.Num() > 0)
@@ -603,9 +586,34 @@ namespace CrossCompiler
 	PRAGMA_ENABLE_DEPRECATION_WARNINGS
 	{
 		// Convert FString to ANSI string and store them as intermediates
-		ConvertFStringViewToAnsiString(ShaderSource, Intermediates->ShaderSource);
-		ConvertFStringToAnsiString(Filename, Intermediates->Filename);
-		ConvertFStringToAnsiString(EntryPoint, Intermediates->EntryPoint);
+		Intermediates->ShaderSource = ShaderSource;
+		Intermediates->Filename = Filename;
+		Intermediates->EntryPoint = EntryPoint;
+
+		// Convert macro definitions map into an array container
+		if (Definitions != nullptr)
+		{
+			ConvertDefineMapToMacroDefines(*Definitions, Intermediates->Defines, Intermediates->DefineRefs);
+		}
+
+		if (ExtraDxcArgs && ExtraDxcArgs->Num() > 0)
+		{
+			ConvertStringArrayToAnsiArray(*ExtraDxcArgs, Intermediates->CustomDxcArgs, Intermediates->CustomDxcArgRefs);
+		}
+
+		// Convert shader stage
+		Intermediates->Stage = ToShaderConductorShaderStage(ShaderStage);
+
+		return true;
+	}
+
+	PRAGMA_DISABLE_DEPRECATION_WARNINGS		// FShaderCompilerDefinitions will be made internal in the future, marked deprecated until then
+	bool FShaderConductorContext::LoadSource(FAnsiStringView ShaderSource, const FString& Filename, const FString& EntryPoint, EShaderFrequency ShaderStage, const FShaderCompilerDefinitions* Definitions, const TArray<FString>* ExtraDxcArgs)
+	PRAGMA_ENABLE_DEPRECATION_WARNINGS
+	{
+		Intermediates->ShaderSource = ShaderSource;
+		Intermediates->Filename = Filename;
+		Intermediates->EntryPoint = EntryPoint;
 
 		// Convert macro definitions map into an array container
 		if (Definitions != nullptr)
@@ -624,6 +632,7 @@ namespace CrossCompiler
 		return true;
 	}
 
+
 	PRAGMA_DISABLE_DEPRECATION_WARNINGS		// FShaderCompilerDefinitions will be made internal in the future, marked deprecated until then
 	bool FShaderConductorContext::LoadSource(const FString& ShaderSource, const FString& Filename, const FString& EntryPoint, EShaderFrequency ShaderStage, const FShaderCompilerDefinitions* Definitions, const TArray<FString>* ExtraDxcArgs)
 	PRAGMA_ENABLE_DEPRECATION_WARNINGS
@@ -635,10 +644,9 @@ namespace CrossCompiler
 	bool FShaderConductorContext::LoadSource(const ANSICHAR* ShaderSource, const ANSICHAR* Filename, const ANSICHAR* EntryPoint, EShaderFrequency ShaderStage, const FShaderCompilerDefinitions* Definitions, const TArray<FString>* ExtraDxcArgs)
 	PRAGMA_ENABLE_DEPRECATION_WARNINGS
 	{
-		// Store ANSI strings as intermediates
-		CopyAnsiString(ShaderSource, Intermediates->ShaderSource);
-		CopyAnsiString(Filename, Intermediates->Filename);
-		CopyAnsiString(EntryPoint, Intermediates->EntryPoint);
+		Intermediates->ShaderSource = ShaderSource;
+		Intermediates->Filename = Filename;
+		Intermediates->EntryPoint = EntryPoint;
 
 		// Convert macro definitions map into an array container
 		if (Definitions != nullptr)
@@ -674,13 +682,15 @@ namespace CrossCompiler
 
 		if (!ResultDesc.hasError && ResultDesc.target.Size() > 0)
 		{
+			FAnsiStringView ResultView(reinterpret_cast<const ANSICHAR*>(ResultDesc.target.Data()), ResultDesc.target.Size());
 			// Copy rewritten HLSL code into intermediate source code.
-			ConvertByteArrayToAnsiString(reinterpret_cast<const ANSICHAR*>(ResultDesc.target.Data()), ResultDesc.target.Size(), Intermediates->ShaderSource);
+			Intermediates->ShaderSource.CopyAnsi(ResultView);
 
 			// If output source is specified, also convert to TCHAR string
 			if (OutSource != nullptr)
 			{
-				*OutSource = ANSI_TO_TCHAR(Intermediates->ShaderSource.GetData());
+				OutSource->Empty();
+				OutSource->Append(ResultView);
 			}
 			bSucceeded = true;
 		}
@@ -766,7 +776,8 @@ namespace CrossCompiler
 			[&OutSource](const void* Data, uint32 Size)
 			{
 				// Convert source buffer to FString
-				ConvertByteArrayToFString(Data, Size, OutSource);
+				FString Converted(reinterpret_cast<const ANSICHAR*>(Data), Size);
+				OutSource = MoveTemp(Converted);
 			}
 		);
 	}
@@ -778,7 +789,8 @@ namespace CrossCompiler
 			[&OutSource](const void* Data, uint32 Size)
 			{
 				// Convert source buffer to ANSI string
-				ConvertByteArrayToAnsiString(reinterpret_cast<const ANSICHAR*>(Data), Size, OutSource);
+				FAnsiString Copy = FAnsiString::ConstructFromPtrSize(reinterpret_cast<const ANSICHAR*>(Data), Size);
+				OutSource = MoveTemp(Copy.GetCharArray());
 			}
 		);
 	}
@@ -850,12 +862,12 @@ namespace CrossCompiler
 
 	const ANSICHAR* FShaderConductorContext::GetSourceString() const
 	{
-		return (Intermediates->ShaderSource.Num() > 0 ? Intermediates->ShaderSource.GetData() : nullptr);
+		return (Intermediates->ShaderSource.AnsiView.Len() > 0 ? Intermediates->ShaderSource.AnsiView.GetData() : nullptr);
 	}
 
 	int32 FShaderConductorContext::GetSourceLength() const
 	{
-		return (Intermediates->ShaderSource.Num() > 0 ? (Intermediates->ShaderSource.Num() - 1) : 0);
+		return Intermediates->ShaderSource.AnsiView.Len();
 	}
 
 	static const TCHAR* GetHlslShaderModelProfile(ShaderConductor::ShaderStage Stage)
@@ -875,8 +887,8 @@ namespace CrossCompiler
 	FString FShaderConductorContext::GenerateDxcArguments(const FShaderConductorOptions& Options) const
 	{
 		FString CmdLineArgs = FString::Printf(
-			TEXT("-E %s -T %s_%d_%d"),
-			ANSI_TO_TCHAR(Intermediates->EntryPoint.GetData()), GetHlslShaderModelProfile(Intermediates->Stage), (int32)Options.ShaderModel.Major, (int32)Options.ShaderModel.Minor
+			TEXT("-E %hs -T %s_%d_%d"),
+			Intermediates->EntryPoint.AnsiView.GetData(), GetHlslShaderModelProfile(Intermediates->Stage), (int32)Options.ShaderModel.Major, (int32)Options.ShaderModel.Minor
 		);
 
 		TArray<const ANSICHAR*> DxcArguments;
@@ -942,7 +954,8 @@ namespace CrossCompiler
 		if (TextOutput.isText && !TextOutput.hasError)
 		{
 			// Convert and return output to ANSI string
-			ConvertByteArrayToAnsiString(reinterpret_cast<const ANSICHAR*>(TextOutput.target.Data()), TextOutput.target.Size(), OutAssemblyText);
+			FAnsiString Copy = FAnsiString::ConstructFromPtrSize(reinterpret_cast<const ANSICHAR*>(TextOutput.target.Data()), TextOutput.target.Size());
+			OutAssemblyText = MoveTemp(Copy.GetCharArray());
 			return true;
 		}
 
diff --git a/Engine/Source/Developer/ShaderCompilerCommon/Public/ShaderConductorContext.h b/Engine/Source/Developer/ShaderCompilerCommon/Public/ShaderConductorContext.h
--- a/Engine/Source/Developer/ShaderCompilerCommon/Public/ShaderConductorContext.h
+++ b/Engine/Source/Developer/ShaderCompilerCommon/Public/ShaderConductorContext.h
@@ -202,6 +202,7 @@ namespace CrossCompiler
 		PRAGMA_DISABLE_DEPRECATION_WARNINGS		// FShaderCompilerDefinitions will be made internal in the future, marked deprecated until then
 		bool LoadSource(const FString& ShaderSource, const FString& Filename, const FString& EntryPoint, EShaderFrequency ShaderStage, const FShaderCompilerDefinitions* Definitions = nullptr, const TArray<FString>* ExtraDxcArgs = nullptr);
 		bool LoadSource(FStringView ShaderSource, const FString& Filename, const FString& EntryPoint, EShaderFrequency ShaderStage, const FShaderCompilerDefinitions* Definitions = nullptr, const TArray<FString>* ExtraDxcArgs = nullptr);
+		bool LoadSource(FAnsiStringView ShaderSource, const FString& Filename, const FString& EntryPoint, EShaderFrequency ShaderStage, const FShaderCompilerDefinitions* Definitions = nullptr, const TArray<FString>* ExtraDxcArgs = nullptr);
 		bool LoadSource(const ANSICHAR* ShaderSource, const ANSICHAR* Filename, const ANSICHAR* EntryPoint, EShaderFrequency ShaderStage, const FShaderCompilerDefinitions* Definitions = nullptr, const TArray<FString>* ExtraDxcArgs = nullptr);
 		PRAGMA_ENABLE_DEPRECATION_WARNINGS
 
