ShaderMinifier - Use memstack for various temporary allocations

diff --git a/Engine/Source/Developer/ShaderCompilerCommon/Private/ShaderMinifier.cpp b/Engine/Source/Developer/ShaderCompilerCommon/Private/ShaderMinifier.cpp
--- a/Engine/Source/Developer/ShaderCompilerCommon/Private/ShaderMinifier.cpp
+++ b/Engine/Source/Developer/ShaderCompilerCommon/Private/ShaderMinifier.cpp
@@ -8,6 +8,7 @@
 #include "Misc/AutomationTest.h"
 #include "String/Find.h"
 #include "Algo/BinarySearch.h"
+#include "Misc/MemStack.h"
 
 #define UE_SHADER_MINIFIER_SSE (PLATFORM_CPU_X86_FAMILY && PLATFORM_ENABLE_VECTORINTRINSICS && PLATFORM_ALWAYS_HAS_SSE4_2)
 
@@ -24,6 +25,9 @@ DEFINE_LOG_CATEGORY_STATIC(LogShaderMinifier, Log, All);
 namespace UE::ShaderMinifier
 {
 
+using FMemStackSetAllocator = TSetAllocator<TSparseArrayAllocator<TMemStackAllocator<>, TMemStackAllocator<>>, TMemStackAllocator<>>;
+using FMemStackAllocator = TMemStackAllocator<>;
+
 static FShaderSource::FViewType SubStrView(FShaderSource::FViewType S, int32 Start)
 {
 	Start = FMath::Min(Start, S.Len());
@@ -1308,7 +1312,7 @@ static TArray<FShaderSource::FViewType> SplitByChar(FShaderSource::FViewType Sou
 	return Result;
 }
 
-static void ExtractIdentifiers(FShaderSource::FViewType InSource, TArray<FShaderSource::FViewType>& Result)
+static void ExtractIdentifiers(FShaderSource::FViewType InSource, TArray<FShaderSource::FViewType, FMemStackAllocator>& Result)
 {
 	FShaderSource::FViewType Source = InSource;
 
@@ -1379,7 +1383,7 @@ static void ExtractIdentifiers(FShaderSource::FViewType InSource, TArray<FShader
 	}
 }
 
-static void ExtractIdentifiers(const FCodeChunk& Chunk, TArray<FShaderSource::FViewType>& Result)
+static void ExtractIdentifiers(const FCodeChunk& Chunk, TArray<FShaderSource::FViewType, FMemStackAllocator>& Result)
 {
 	for (const FCodeBlock& Block : Chunk.Blocks)
 	{
@@ -1454,7 +1458,7 @@ struct FCasedStringViewKeyFuncs : public DefaultKeyFuncs<FShaderSource::FViewTyp
 	}
 };
 
-static void BuildLineBreakMap(FShaderSource::FViewType Source, TArray<int32>& OutLineBreakMap)
+static void BuildLineBreakMap(FShaderSource::FViewType Source, TArray<int32, FMemStackAllocator>& OutLineBreakMap)
 {
 	OutLineBreakMap.Reset();
 
@@ -1509,7 +1513,7 @@ static int32 FindLineDirective(const TArray<FShaderSource::FViewType>& LineDirec
 	return FoundIndex - 1;
 }
 
-static int32 FindLineNumber(FShaderSource::FViewType Source, const TArray<int32>& LineBreakMap, const FShaderSource::CharType* Ptr)
+static int32 FindLineNumber(FShaderSource::FViewType Source, const TArray<int32, FMemStackAllocator>& LineBreakMap, const FShaderSource::CharType* Ptr)
 {
 	if (Ptr < Source.GetData() || Ptr >= Source.GetData() + Source.Len())
 	{
@@ -1581,16 +1585,18 @@ static void CloseNamespace(FShaderSource::FStringType& OutputStream, const FName
 
 static FShaderSource::FStringType MinifyShader(const FParsedShader& Parsed, TConstArrayView<FShaderSource::FViewType> RequiredSymbols, EMinifyShaderFlags Flags, FDiagnostics& Diagnostics)
 {
+	FMemMark Mark(FMemStack::Get());
+
 	FShaderSource::FStringType OutputStream;
 
 	OutputStream.Reserve(Parsed.Source.Len() / 3); // Heuristic pre-allocation based on average measured reduced code size
 
-	TSet<FShaderSource::FViewType, FCasedStringViewKeyFuncs, FDefaultSetAllocator> RelevantIdentifiers;
+	TSet<FShaderSource::FViewType, FCasedStringViewKeyFuncs, FMemStackSetAllocator> RelevantIdentifiers;
 
-	TSet<const FCodeChunk*> RelevantChunks;
-	TSet<FShaderSource::FViewType, FCasedStringViewKeyFuncs, FDefaultSetAllocator>  ProcessedIdentifiers;
+	TSet<const FCodeChunk*, DefaultKeyFuncs<const FCodeChunk*>, FMemStackSetAllocator> RelevantChunks;
+	TSet<FShaderSource::FViewType, FCasedStringViewKeyFuncs, FMemStackSetAllocator> ProcessedIdentifiers;
 
-	TArray<const FCodeChunk*> PendingChunks;
+	TArray<const FCodeChunk*, FMemStackAllocator> PendingChunks;
 
 	for (FShaderSource::FViewType Entry : RequiredSymbols)
 	{
@@ -1749,9 +1755,9 @@ static FShaderSource::FStringType MinifyShader(const FParsedShader& Parsed, TCon
 		return {};
 	}
 
-	TArray<FShaderSource::FViewType> TempIdentifiers;
+	TArray<FShaderSource::FViewType, FMemStackAllocator> TempIdentifiers;
 
-	TMap<FShaderSource::FViewType, TArray<const FCodeChunk*>, FDefaultSetAllocator, FCasedStringViewKeyFuncs> ChunksByIdentifier;
+	TMap<FShaderSource::FViewType, TArray<const FCodeChunk*, FMemStackAllocator>, FMemStackSetAllocator, FCasedStringViewKeyFuncs> ChunksByIdentifier;
 	for (const FCodeChunk& Chunk : Parsed.Chunks)
 	{
 		for (const FCodeBlock& Block : Chunk.Blocks)
@@ -1799,7 +1805,7 @@ static FShaderSource::FStringType MinifyShader(const FParsedShader& Parsed, TCon
 		}
 	}
 
-	TMap<const FCodeChunk*, const FCodeChunk*> ChunkRequestedBy;
+	TMap<const FCodeChunk*, const FCodeChunk*, FMemStackSetAllocator> ChunkRequestedBy;
 
 	while (!PendingChunks.IsEmpty())
 	{
@@ -1911,7 +1917,7 @@ static FShaderSource::FStringType MinifyShader(const FParsedShader& Parsed, TCon
 		OutputStream.AppendChar('\n');
 	}
 
-	TArray<int32> LineBreakMap;
+	TArray<int32, FMemStackAllocator> LineBreakMap;
 	const TArray<FShaderSource::FViewType>& LineDirectives = Parsed.LineDirectives;
 	if (EnumHasAnyFlags(Flags, EMinifyShaderFlags::OutputLines))
 	{
@@ -2151,7 +2157,8 @@ bool FShaderMinifierParserTest::RunTest(const FString& Parameters)
 	}
 
 	{
-		TArray<FShaderSource::FViewType> R;
+		FMemMark Mark(FMemStack::Get());
+		TArray<FShaderSource::FViewType, FMemStackAllocator> R;
 		FShaderSource S(SHADER_SOURCE_LITERAL("Hello[World]; Foo[0];\n"));
 		ExtractIdentifiers(S.GetView(), R);
 		if (TestEqual(TEXT("ExtractIdentifiers1: Num"), R.Num(), 3))
@@ -2163,7 +2170,8 @@ bool FShaderMinifierParserTest::RunTest(const FString& Parameters)
 	}
 
 	{
-		TArray<FShaderSource::FViewType> R;
+		FMemMark Mark(FMemStack::Get());
+		TArray<FShaderSource::FViewType, FMemStackAllocator> R;
 		FShaderSource S(SHADER_SOURCE_LITERAL("#line 0\nStructuredBuffer<uint4> Blah : register(t0, space123);#line 1\n#pragma foo\n"));
 		ExtractIdentifiers(S.GetView(), R);
 		if (TestEqual(TEXT("ExtractIdentifiers2: Num"), R.Num(), 6))
