Fixed RDNA2 disassembler to correctly handle v_madak_f32 and v_madmk_f32 advancing with a trailing literal

diff --git a/Engine/Source/Developer/ShaderCompilerCommon/Private/ISAParser/QuadModeDetect.cpp b/Engine/Source/Developer/ShaderCompilerCommon/Private/ISAParser/QuadModeDetect.cpp
--- a/Engine/Source/Developer/ShaderCompilerCommon/Private/ISAParser/QuadModeDetect.cpp
+++ b/Engine/Source/Developer/ShaderCompilerCommon/Private/ISAParser/QuadModeDetect.cpp
@@ -171,6 +171,19 @@ struct FDebugParserRDNA2
 	uint32 AdjustAdvance(EInstructionType InstructionType, const uint32* InstPtr)
 	{
 		uint32 AdvanceAmount = 0u;
+
+		if (InstructionType == EInstructionType::VOP2)
+		{
+			const FInstVOP2& Inst = *reinterpret_cast<const FInstVOP2*>(InstPtr);
+			const RDNA2::EVOP2Ops Op = (const RDNA2::EVOP2Ops)Inst.OP;
+
+			// These instructions always have a literal constant that follows
+			if (Op == RDNA2::EVOP2Ops::v_madmk_f32 || Op == RDNA2::EVOP2Ops::v_madak_f32)
+			{
+				AdvanceAmount++;
+			}
+		}
+
 		return AdvanceAmount;
 	}
 
@@ -479,6 +492,19 @@ struct FQuadModeParserRDNA2
 	uint32 AdjustAdvance(EInstructionType InstructionType, const uint32* InstPtr)
 	{
 		uint32 AdvanceAmount = 0u;
+
+		if (InstructionType == EInstructionType::VOP2)
+		{
+			const FInstVOP2& Inst = *reinterpret_cast<const FInstVOP2*>(InstPtr);
+			const RDNA2::EVOP2Ops Op = (const RDNA2::EVOP2Ops)Inst.OP;
+
+			// These instructions always have a literal constant that follows
+			if (Op == RDNA2::EVOP2Ops::v_madmk_f32 || Op == RDNA2::EVOP2Ops::v_madak_f32)
+			{
+				AdvanceAmount++;
+			}
+		}
+
 		return AdvanceAmount;
 	}
 
diff --git a/Engine/Source/Developer/ShaderCompilerCommon/Private/ISAParser/RDNA2_ISA.h b/Engine/Source/Developer/ShaderCompilerCommon/Private/ISAParser/RDNA2_ISA.h
--- a/Engine/Source/Developer/ShaderCompilerCommon/Private/ISAParser/RDNA2_ISA.h
+++ b/Engine/Source/Developer/ShaderCompilerCommon/Private/ISAParser/RDNA2_ISA.h
@@ -807,6 +807,9 @@ enum class EVOP2Ops : uint16
 	v_or_b32 = 28,
 	v_xor_b32 = 29,
 	v_xnor_b32 = 30,
+	v_mac_f32 = 31,
+	v_madmk_f32 = 32,
+	v_madak_f32 = 33,
 	v_add_nc_u32 = 37,
 	v_sub_nc_u32 = 38,
 	v_subrev_nc_u32 = 39,
@@ -862,6 +865,9 @@ const char* ToString(EVOP2Ops Op)
 		OP_TO_STRING_CASE(v_or_b32);
 		OP_TO_STRING_CASE(v_xor_b32);
 		OP_TO_STRING_CASE(v_xnor_b32);
+		OP_TO_STRING_CASE(v_mac_f32);
+		OP_TO_STRING_CASE(v_madmk_f32);
+		OP_TO_STRING_CASE(v_madak_f32);
 		OP_TO_STRING_CASE(v_add_nc_u32);
 		OP_TO_STRING_CASE(v_sub_nc_u32);
 		OP_TO_STRING_CASE(v_subrev_nc_u32);
@@ -1290,6 +1296,7 @@ const char* ToString(EVOPCOps Op)
 enum class EVOP3ABOps : uint16
 {
 	v_fma_legacy_f32 = 320,
+	v_mad_f32 = 321,
 	v_mad_i32_i24 = 322,
 	v_mad_u32_u24 = 323,
 	v_cubeid_f32 = 324,
@@ -1432,6 +1439,7 @@ const char* ToString(EVOP3ABOps Op)
 	switch (Op)
 	{
 		OP_TO_STRING_CASE(v_fma_legacy_f32);
+		OP_TO_STRING_CASE(v_mad_f32);
 		OP_TO_STRING_CASE(v_mad_i32_i24);
 		OP_TO_STRING_CASE(v_mad_u32_u24);
 		OP_TO_STRING_CASE(v_cubeid_f32);
@@ -2311,7 +2319,6 @@ enum class ELDSGDSOps : uint16
 	ds_read_b128 = 255,
 };
 
-// TODO: Compare against RDNA 2
 const char* ToString(ELDSGDSOps Op)
 {
 #define OP_TO_STRING_CASE(x) case ELDSGDSOps::x: return #x
