Fixed some additional ISA issues by adding v_fmaak_f32, v_fmamk_f32, v_fmamk_f16, and v_fmaak_f16 to the trailing literal check

diff --git a/Engine/Source/Developer/ShaderCompilerCommon/Private/ISAParser/QuadModeDetect.cpp b/Engine/Source/Developer/ShaderCompilerCommon/Private/ISAParser/QuadModeDetect.cpp
--- a/Engine/Source/Developer/ShaderCompilerCommon/Private/ISAParser/QuadModeDetect.cpp
+++ b/Engine/Source/Developer/ShaderCompilerCommon/Private/ISAParser/QuadModeDetect.cpp
@@ -16,7 +16,7 @@ struct FDebugParserRDNA1
 			const RDNA1::EVOP2Ops Op = (const RDNA1::EVOP2Ops)Inst.OP;
 
 			// These instructions always have a literal constant that follows
-			if (Op == RDNA1::EVOP2Ops::v_madmk_f32 || Op == RDNA1::EVOP2Ops::v_madak_f32)
+			if (RDNA1::HasTrailingLiteral(Op))
 			{
 				AdvanceAmount++;
 			}
@@ -178,7 +178,7 @@ struct FDebugParserRDNA2
 			const RDNA2::EVOP2Ops Op = (const RDNA2::EVOP2Ops)Inst.OP;
 
 			// These instructions always have a literal constant that follows
-			if (Op == RDNA2::EVOP2Ops::v_madmk_f32 || Op == RDNA2::EVOP2Ops::v_madak_f32)
+			if (RDNA2::HasTrailingLiteral(Op))
 			{
 				AdvanceAmount++;
 			}
@@ -343,7 +343,7 @@ struct FQuadModeParserRDNA1
 			const RDNA1::EVOP2Ops Op = (const RDNA1::EVOP2Ops)Inst.OP;
 
 			// These instructions always have a literal constant that follows
-			if (Op == RDNA1::EVOP2Ops::v_madmk_f32 || Op == RDNA1::EVOP2Ops::v_madak_f32)
+			if (RDNA1::HasTrailingLiteral(Op))
 			{
 				AdvanceAmount++;
 			}
@@ -499,7 +499,7 @@ struct FQuadModeParserRDNA2
 			const RDNA2::EVOP2Ops Op = (const RDNA2::EVOP2Ops)Inst.OP;
 
 			// These instructions always have a literal constant that follows
-			if (Op == RDNA2::EVOP2Ops::v_madmk_f32 || Op == RDNA2::EVOP2Ops::v_madak_f32)
+			if (RDNA2::HasTrailingLiteral(Op))
 			{
 				AdvanceAmount++;
 			}
diff --git a/Engine/Source/Developer/ShaderCompilerCommon/Private/ISAParser/RDNA1_ISA.h b/Engine/Source/Developer/ShaderCompilerCommon/Private/ISAParser/RDNA1_ISA.h
--- a/Engine/Source/Developer/ShaderCompilerCommon/Private/ISAParser/RDNA1_ISA.h
+++ b/Engine/Source/Developer/ShaderCompilerCommon/Private/ISAParser/RDNA1_ISA.h
@@ -2578,6 +2578,23 @@ const char* ToString(EGLOBALOps Op)
 #undef OP_TO_STRING_CASE
 }
 
+inline bool HasTrailingLiteral(RDNA1::EVOP2Ops Op)
+{
+	switch (Op)
+	{
+	case RDNA1::EVOP2Ops::v_madmk_f32:
+	case RDNA1::EVOP2Ops::v_madak_f32:
+	case RDNA1::EVOP2Ops::v_fmamk_f32:
+	case RDNA1::EVOP2Ops::v_fmaak_f32:
+	case RDNA1::EVOP2Ops::v_fmamk_f16:
+	case RDNA1::EVOP2Ops::v_fmaak_f16:
+		return true;
+
+	default:
+		return false;
+	}
+}
+
 void PrintSMEM(const FInstSMEM& Inst)
 {
 	ESMEMOps Op = (ESMEMOps)Inst.OP;
diff --git a/Engine/Source/Developer/ShaderCompilerCommon/Private/ISAParser/RDNA2_ISA.h b/Engine/Source/Developer/ShaderCompilerCommon/Private/ISAParser/RDNA2_ISA.h
--- a/Engine/Source/Developer/ShaderCompilerCommon/Private/ISAParser/RDNA2_ISA.h
+++ b/Engine/Source/Developer/ShaderCompilerCommon/Private/ISAParser/RDNA2_ISA.h
@@ -2586,6 +2586,23 @@ const char* ToString(EGLOBALOps Op)
 #undef OP_TO_STRING_CASE
 }
 
+inline bool HasTrailingLiteral(RDNA2::EVOP2Ops Op)
+{
+	switch (Op)
+	{
+	case RDNA2::EVOP2Ops::v_madmk_f32:
+	case RDNA2::EVOP2Ops::v_madak_f32:
+	case RDNA2::EVOP2Ops::v_fmamk_f32:
+	case RDNA2::EVOP2Ops::v_fmaak_f32:
+	case RDNA2::EVOP2Ops::v_fmamk_f16:
+	case RDNA2::EVOP2Ops::v_fmaak_f16:
+		return true;
+
+	default:
+		return false;
+	}
+}
+
 void PrintSMEM(const FInstSMEM& Inst)
 {
 	ESMEMOps Op = (ESMEMOps)Inst.OP;
