Add workgraph node support to shader minifier.
Missing functionality was that node functions can have a NodeID that is used to reference them.
The NodeID is set in an attribute on the function.
The fix is to always extract identifiers from function attributes.
This could be made more conservative by special casing to only NodeID function attributes. I'm not sure if the added complexity is worth it.
Note also that this works for node arrays, but doesn't remove out of array references when the NodeArraySize attribute is used. That is quite a special case though (different array slices from a shared array definition).

diff --git a/Engine/Source/Developer/ShaderCompilerCommon/Private/ShaderMinifier.cpp b/Engine/Source/Developer/ShaderCompilerCommon/Private/ShaderMinifier.cpp
--- a/Engine/Source/Developer/ShaderCompilerCommon/Private/ShaderMinifier.cpp
+++ b/Engine/Source/Developer/ShaderCompilerCommon/Private/ShaderMinifier.cpp
@@ -1767,7 +1767,7 @@ static FShaderSource::FStringType MinifyShader(const FParsedShader& Parsed, TCon
 				continue;
 			}
 
-			if (Chunk.Type == ECodeChunkType::Function && Block.Type != EBlockType::Name)
+			if (Chunk.Type == ECodeChunkType::Function && Block.Type != EBlockType::Name && Block.Type != EBlockType::Attribute)
 			{
 				continue;
 			}
@@ -1782,7 +1782,18 @@ static FShaderSource::FStringType MinifyShader(const FParsedShader& Parsed, TCon
 				continue;
 			}
 
-			if ((Chunk.Type == ECodeChunkType::CBuffer || Chunk.Type == ECodeChunkType::Enum) && Block.Type == EBlockType::Body)
+			if (Chunk.Type == ECodeChunkType::Operator
+				&& Block.Type != EBlockType::Type
+				&& Block.Type != EBlockType::Args)
+			{
+				continue;
+			}
+
+			bool bExtractIdentifiers = (Chunk.Type == ECodeChunkType::CBuffer && Block.Type == EBlockType::Body);
+			bExtractIdentifiers |= (Chunk.Type == ECodeChunkType::Enum && Block.Type == EBlockType::Body);
+			bExtractIdentifiers |= (Chunk.Type == ECodeChunkType::Function && Block.Type == EBlockType::Attribute);
+
+			if (bExtractIdentifiers)
 			{
 				TempIdentifiers.Reset();
 				ExtractIdentifiers(Block, TempIdentifiers);
@@ -1794,13 +1805,6 @@ static FShaderSource::FStringType MinifyShader(const FParsedShader& Parsed, TCon
 				continue;
 			}
 
-			if (Chunk.Type == ECodeChunkType::Operator
-				&& Block.Type != EBlockType::Type
-				&& Block.Type != EBlockType::Args)
-			{
-				continue;
-			}
-
 			ChunksByIdentifier.FindOrAdd(Block).Push(&Chunk);
 		}
 	}
@@ -2461,6 +2465,20 @@ bool FShaderMinifierParserTest::RunTest(const FString& Parameters)
 		}
 	}
 
+	{
+		FShaderSource S(SHADER_SOURCE_LITERAL("[Shader(\"node\")] [NodeID(\"WorkgraphNodeArray\", 0)] void WorkgraphNodeArray_0() {}"));
+		auto P = ParseShader(S);
+		TestEqual(TEXT("ParseShader: workgraph node: num chunks"), P.Chunks.Num(), 1);
+		if (P.Chunks.Num() == 1)
+		{
+			TestEqual(TEXT("ParseShader: workgraph node, chunk type"), P.Chunks[0].Type, ECodeChunkType::Function);
+			if (TestEqual(TEXT("ParseShader: workgraph node: chunk 0: num blocks"), P.Chunks[0].Blocks.Num(), 6))
+			{
+				TestEqual(TEXT("ParseShader: workgraph node: chunk 0: block 1: block type"), P.Chunks[0].Blocks[1].Type, EBlockType::Attribute);
+			}
+		}
+	}
+
 	int32 NumErrors = ExecutionInfo.GetErrorTotal();
 
 	return NumErrors == 0;
@@ -2598,10 +2616,20 @@ enum EEnumUnused
 	ENUM_UNUSED_PART_2,
 };
 
+struct FWorkgraphRecord
+{
+	int Foo;
+};
+
 // Test comment 2
 [numthreads(1,1,1)]
 // Comment during function declaration
-void MainCS()
+void MainCS(
+	NodeOutput<FWorkgraphRecord> WorkgraphNode,
+	EmptyNodeOutput WorkgraphNodeAliased,
+	[NodeID("WorkgraphNodeArray", 0)]
+	EmptyNodeOutputArray WorkgraphNodeArray
+	)
 {
 	using namespace NS1::NS2;
 	using namespace NS3;
@@ -2612,6 +2640,29 @@ void MainCS()
 	float D = TypedefUsedBuffer[ENUM_USED_PART_1].Foo;
 	OutputBuffer[0] = A + B + D;
 }
+
+[Shader("node")]
+void WorkgraphNode()
+{
+}
+
+[Shader("node")]
+[NodeID("WorkgraphNodeAliased")]
+void WorkgraphNodeFunction()
+{
+}
+
+[Shader("node")]
+[NodeID("WorkgraphNodeArray", 0)]
+void WorkgraphNodeArray_0()
+{
+}
+
+[Shader("node")]
+[NodeID("WorkgraphNodeArray", 1)]
+void WorkgraphNodeArray_1()
+{
+}
 )"));
 
 	auto ChunkPresent = [](const FParsedShader& Parsed, FShaderSource::FViewType Name)
@@ -2672,6 +2723,11 @@ void MainCS()
 		TestTrue(TEXT("MinifyShader: MainCS: contains EEnumUsed"), MinifiedParsed.Source.Contains(SHADER_SOURCE_LITERAL("EEnumUsed")));
 		TestTrue(TEXT("MinifyShader: MainCS: contains ENUM_USED_PART_1"), MinifiedParsed.Source.Contains(SHADER_SOURCE_LITERAL("ENUM_USED_PART_1")));
 		TestTrue(TEXT("MinifyShader: MainCS: contains ENUM_USED_PART_2"), MinifiedParsed.Source.Contains(SHADER_SOURCE_LITERAL("ENUM_USED_PART_2")));
+		TestTrue(TEXT("MinifyShader: MainCS: contains struct FWorkgraphRecord"), MinifiedParsed.Source.Contains(SHADER_SOURCE_LITERAL("struct FWorkgraphRecord")));
+		TestTrue(TEXT("MinifyShader: MainCS: contains void WorkgraphNode()"), MinifiedParsed.Source.Contains(SHADER_SOURCE_LITERAL("void WorkgraphNode()")));
+		TestTrue(TEXT("MinifyShader: MainCS: contains void WorkgraphNodeFunction()"), MinifiedParsed.Source.Contains(SHADER_SOURCE_LITERAL("void WorkgraphNodeFunction()")));
+		TestTrue(TEXT("MinifyShader: MainCS: contains void WorkgraphNodeArray_0()"), MinifiedParsed.Source.Contains(SHADER_SOURCE_LITERAL("void WorkgraphNodeArray_0()")));
+		TestTrue(TEXT("MinifyShader: MainCS: contains void WorkgraphNodeArray_1()"), MinifiedParsed.Source.Contains(SHADER_SOURCE_LITERAL("void WorkgraphNodeArray_1()")));
 
 		// Expect false:
 		TestFalse(TEXT("MinifyShader: MainCS: contains UnreferencedFunction"), ChunkPresent(MinifiedParsed, SHADER_SOURCE_LITERAL("UnreferencedFunction")));
