[shaders] add additional data (serialized env, cmdline, etc) to all copies of source dumped when "detailed source" debug dump option is enabled

diff --git a/Engine/Source/Developer/ShaderCompilerCommon/Private/ShaderCompilerCommon.cpp b/Engine/Source/Developer/ShaderCompilerCommon/Private/ShaderCompilerCommon.cpp
--- a/Engine/Source/Developer/ShaderCompilerCommon/Private/ShaderCompilerCommon.cpp
+++ b/Engine/Source/Developer/ShaderCompilerCommon/Private/ShaderCompilerCommon.cpp
@@ -1798,10 +1798,10 @@ namespace UE::ShaderCompilerCommon
 			}
 			
 			PrefixedOptions.FilenamePrefix = *StrippedPrefix;
-			FFileHelper::SaveStringToFile(PreprocessOutput.GetSourceViewWide(), *PrefixedOptions.GetDebugShaderPath(Input));
+			FFileHelper::SaveStringToFile(GetDebugShaderContents(Input, PreprocessOutput.GetSourceViewWide(), PrefixedOptions), *PrefixedOptions.GetDebugShaderPath(Input));
 
 			PrefixedOptions.FilenamePrefix = *PreprocessedPrefix;
-			FFileHelper::SaveStringToFile(PreprocessOutput.GetUnstrippedSourceView(), *PrefixedOptions.GetDebugShaderPath(Input));
+			FFileHelper::SaveStringToFile(GetDebugShaderContents(Input, PreprocessOutput.GetUnstrippedSourceView(), PrefixedOptions), *PrefixedOptions.GetDebugShaderPath(Input));
 		}
 		if (Output.ModifiedShaderSource.IsEmpty())
 		{
