Removing extra '\0' from files output from ShaderConductor.
DxcCreateBlob apparently force adds this on creation of strings, so we need to account for that when creating a view of the result.

diff --git a/Engine/Source/Developer/ShaderCompilerCommon/Private/ShaderConductorContext.cpp b/Engine/Source/Developer/ShaderCompilerCommon/Private/ShaderConductorContext.cpp
--- a/Engine/Source/Developer/ShaderCompilerCommon/Private/ShaderConductorContext.cpp
+++ b/Engine/Source/Developer/ShaderCompilerCommon/Private/ShaderConductorContext.cpp
@@ -680,9 +680,12 @@ namespace CrossCompiler
 		ShaderConductor::Compiler::ResultDesc ResultDesc;
 		ScRewriteWrapper(ScSourceDesc, ScOptions, ResultDesc);
 
-		if (!ResultDesc.hasError && ResultDesc.target.Size() > 0)
+		if (!ResultDesc.hasError && ResultDesc.target.Size() > 1)
 		{
-			FAnsiStringView ResultView(reinterpret_cast<const ANSICHAR*>(ResultDesc.target.Data()), ResultDesc.target.Size());
+			// Note: We don't want to include the '\0' included in the result string (thanks to DxcCreateBlob), hence the -1
+			check(reinterpret_cast<const ANSICHAR*>(ResultDesc.target.Data())[ResultDesc.target.Size() - 1] == '\0');
+			FAnsiStringView ResultView(reinterpret_cast<const ANSICHAR*>(ResultDesc.target.Data()), ResultDesc.target.Size() - 1);
+
 			// Copy rewritten HLSL code into intermediate source code.
 			Intermediates->ShaderSource.CopyAnsi(ResultView);
 
