[shaders]
- modify which version of source is dumped as a debug artifact by default to be the final preprocessed source instead of the stripped version; this is more useful when debugging in directcompile mode
- re-add the dump of defines as commented code to this artifact, which was inadvertantly removed in a previous CL (bad merge)
- make the "detailed" source dump add the stripped version, and the version modified by the compile step (if it exists)
- only add the additional debug data to the stripped version (which is still something that SCW can process successfully) but remove it from the "modified" version (which typically will cause errors when running in directcompile mode)

diff --git a/Engine/Source/Developer/ShaderCompilerCommon/Private/ShaderCompilerCommon.cpp b/Engine/Source/Developer/ShaderCompilerCommon/Private/ShaderCompilerCommon.cpp
--- a/Engine/Source/Developer/ShaderCompilerCommon/Private/ShaderCompilerCommon.cpp
+++ b/Engine/Source/Developer/ShaderCompilerCommon/Private/ShaderCompilerCommon.cpp
@@ -1776,27 +1776,30 @@ namespace UE::ShaderCompilerCommon
 			FDebugShaderDataOptions PrefixedOptions(Options);
 			uint32 SlackLen = Options.FilenamePrefix ? FCString::Strlen(Options.FilenamePrefix) : 0;
 			FString StrippedPrefix(TEXT("Stripped_"), SlackLen);
-			FString PreprocessedPrefix(TEXT("Preprocessed_"), SlackLen);
+			FString ModifiedPrefix(TEXT("CompileModified_"), SlackLen);
 			if (Options.FilenamePrefix)
 			{
 				StrippedPrefix += Options.FilenamePrefix;
-				PreprocessedPrefix += Options.FilenamePrefix;
+				ModifiedPrefix += Options.FilenamePrefix;
 			}
 			
 			PrefixedOptions.FilenamePrefix = *StrippedPrefix;
 			FFileHelper::SaveStringToFile(GetDebugShaderContents(Input, PreprocessOutput.GetSourceViewWide(), PrefixedOptions), *PrefixedOptions.GetDebugShaderPath(Input));
 
-			PrefixedOptions.FilenamePrefix = *PreprocessedPrefix;
-			FFileHelper::SaveStringToFile(GetDebugShaderContents(Input, PreprocessOutput.GetUnstrippedSourceView(), PrefixedOptions), *PrefixedOptions.GetDebugShaderPath(Input));
-		}
-		if (Output.ModifiedShaderSource.IsEmpty())
-		{
-			DumpDebugShaderData(Input, PreprocessOutput.GetSourceViewWide(), Options);
-		}
-		else
-		{
-			DumpDebugShaderData(Input, FStringView(Output.ModifiedShaderSource), Options);
+			if (!Output.ModifiedShaderSource.IsEmpty())
+			{
+				// intentionally dumping this copy as-is rather than modifying via GetDebugShaderContents; this is not likely to be something that can be compiled in
+				// directcompile mode for debugging purposes, so we don't append the additional data that this requires.
+				PrefixedOptions.FilenamePrefix = *ModifiedPrefix;
+				FFileHelper::SaveStringToFile(Output.ModifiedShaderSource, *PrefixedOptions.GetDebugShaderPath(Input));
+			}
 		}
+
+		// Always output the preprocessed source (prior to stripping); this is guaranteed to work when passed to SCW in "directcompile" mode
+		// and is also more useful for error reporting (the stripped version of the source has line directives removed, and the error/warning remapping
+		// we perform is only done in the cooker process when processing completed jobs.
+		DumpDebugShaderData(Input, PreprocessOutput.GetUnstrippedSourceView(), Options);
+
 		FFileHelper::SaveStringToFile(Output.OutputHash.ToString(), *GetDebugFileName(Input, Options, TEXT("OutputHash.txt")), FFileHelper::EEncodingOptions::ForceAnsi);
 
 		if (EnumHasAnyFlags(Input.DebugInfoFlags, EShaderDebugInfoFlags::Diagnostics))
@@ -1883,7 +1886,7 @@ namespace UE::ShaderCompilerCommon
 			MergedEnvironment.Merge(*Input.SharedEnvironment);
 		}
 
-		FString Contents = Options.AppendPreSource ? Options.AppendPreSource() : FString();
+		FString Contents = MergedEnvironment.GetDefinitionsAsCommentedCode();
 
 		if (Options.AppendPreSource)
 		{
