[shaders] Support _Static_assert statement in FHlslParser for HLSL2021.

_Static_assert statements (see HLSLStaticAssert.h) in HLSL were not handled in the HLSL parser used to remove unused interpolators.

#rnx

diff --git a/Engine/Source/Developer/ShaderCompilerCommon/Private/HlslAST.cpp b/Engine/Source/Developer/ShaderCompilerCommon/Private/HlslAST.cpp
--- a/Engine/Source/Developer/ShaderCompilerCommon/Private/HlslAST.cpp
+++ b/Engine/Source/Developer/ShaderCompilerCommon/Private/HlslAST.cpp
@@ -531,13 +531,31 @@ namespace CrossCompiler
 		{
 			for (auto* Statement : Statements)
 			{
-				if (Statement)
-				{
-					delete Statement;
-				}
+				delete Statement;
 			}
 		}
 
+		FStaticAssertStatement::FStaticAssertStatement(FLinearAllocator* InAllocator, const FSourceInfo& InInfo) :
+			FNode(InAllocator, InInfo),
+			Condition(nullptr)
+		{
+		}
+
+		FStaticAssertStatement::~FStaticAssertStatement()
+		{
+			delete Condition;
+		}
+
+		void FStaticAssertStatement::Write(FASTWriter& Writer) const
+		{
+			Writer.DoIndent();
+			Writer << *Keyword << TEXT("(");
+			checkf(Condition != nullptr, TEXT("Cannot write %s()-statement in HLSL without a conditional expression"), *Keyword);
+			Condition->Write(Writer);
+			Writer << TEXT(", \"") << *Message << TEXT("\");\n");
+			checkf(false, TEXT("TEST: %s()"), *Keyword);
+		}
+
 		FFunctionDefinition::FFunctionDefinition(FLinearAllocator* InAllocator, const FSourceInfo& InInfo) :
 			FNode(InAllocator, InInfo),
 			Prototype(nullptr),
diff --git a/Engine/Source/Developer/ShaderCompilerCommon/Private/HlslAST.h b/Engine/Source/Developer/ShaderCompilerCommon/Private/HlslAST.h
--- a/Engine/Source/Developer/ShaderCompilerCommon/Private/HlslAST.h
+++ b/Engine/Source/Developer/ShaderCompilerCommon/Private/HlslAST.h
@@ -474,6 +474,18 @@ namespace CrossCompiler
 			virtual FCompoundStatement* AsCompoundStatement() override { return this; }
 		};
 
+		struct FStaticAssertStatement : public FNode
+		{
+			FStaticAssertStatement(FLinearAllocator* InAllocator, const FSourceInfo& InInfo);
+			~FStaticAssertStatement();
+
+			virtual void Write(FASTWriter& Writer) const override;
+
+			FString Keyword; // Either "_Static_assert" (DXC) or "static_assert" (C++11 like compilers)
+			FExpression* Condition;
+			FString Message;
+		};
+
 		struct FSemanticSpecifier : public FNode
 		{
 			enum class ESpecType
diff --git a/Engine/Source/Developer/ShaderCompilerCommon/Private/HlslLexer.cpp b/Engine/Source/Developer/ShaderCompilerCommon/Private/HlslLexer.cpp
--- a/Engine/Source/Developer/ShaderCompilerCommon/Private/HlslLexer.cpp
+++ b/Engine/Source/Developer/ShaderCompilerCommon/Private/HlslLexer.cpp
@@ -440,6 +440,8 @@ namespace CrossCompiler
 				InsertToken(TEXT("typedef"), EHlslToken::Typedef);
 				InsertToken(TEXT("packoffset"), EHlslToken::PackOffset);
 				InsertToken(TEXT("operator"), EHlslToken::Operator);
+				InsertToken(TEXT("_Static_assert"), EHlslToken::StaticAssert); // HLSL2021 adopted C11 '_Static_assert'-statements
+				InsertToken(TEXT("static_assert"), EHlslToken::StaticAssert); // Some shader compilers support C++11 'static_assert'-statements
 			}
 		} GStaticInitializer;
 	}
diff --git a/Engine/Source/Developer/ShaderCompilerCommon/Private/HlslLexer.h b/Engine/Source/Developer/ShaderCompilerCommon/Private/HlslLexer.h
--- a/Engine/Source/Developer/ShaderCompilerCommon/Private/HlslLexer.h
+++ b/Engine/Source/Developer/ShaderCompilerCommon/Private/HlslLexer.h
@@ -282,6 +282,7 @@ namespace CrossCompiler
 		Literal,
 		BoolConstant,
 		StringConstant,	// C-style "string"
+		StaticAssert, // _Static_assert (HLSL), static_assert (C++11)
 	};
 
 	enum class ELiteralType
diff --git a/Engine/Source/Developer/ShaderCompilerCommon/Private/HlslParser.cpp b/Engine/Source/Developer/ShaderCompilerCommon/Private/HlslParser.cpp
--- a/Engine/Source/Developer/ShaderCompilerCommon/Private/HlslParser.cpp
+++ b/Engine/Source/Developer/ShaderCompilerCommon/Private/HlslParser.cpp
@@ -1488,6 +1488,60 @@ Done:
 		return Result;
 	}
 
+	EParseResult ParseStaticAssertStatement(FHlslParser& Parser, FLinearAllocator* Allocator, AST::FNode** OutDeclaration)
+	{
+		const FHlslToken* CurrentToken = Parser.Scanner.GetCurrentToken();
+		check(CurrentToken);
+
+		AST::FStaticAssertStatement* Statement = new(Allocator) AST::FStaticAssertStatement(Allocator, CurrentToken->SourceInfo);
+
+		// Store statement keyword to support _Static_assert (HLSL) and static_assert (C++11 like compilers)
+		Statement->Keyword = CurrentToken->String;
+
+		if (!Parser.Scanner.MatchToken(EHlslToken::LeftParenthesis))
+		{
+			Parser.Scanner.SourceError(TEXT("'(' expected!\n"));
+			return ParseResultError();
+		}
+
+		// Parse assertion conditional expression
+		if (ParseExpression(Parser.Scanner, Parser.CurrentScope, 0, Allocator, &Statement->Condition) != EParseResult::Matched)
+		{
+			Parser.Scanner.SourceError(TEXT("Expression expected"));
+			return ParseResultError();
+		}
+
+		// Parse message string literal
+		if (!Parser.Scanner.MatchToken(EHlslToken::Comma))
+		{
+			Parser.Scanner.SourceError(TEXT("',' expected!\n"));
+			return ParseResultError();
+		}
+
+		const FHlslToken* MessageToken = Parser.Scanner.GetCurrentToken();
+		if (MessageToken->Token != EHlslToken::StringConstant)
+		{
+			Parser.Scanner.SourceError(TEXT("')' expected!\n"));
+			return ParseResultError();
+		}
+
+		Statement->Message = MessageToken->String;
+
+		if (!Parser.Scanner.MatchToken(EHlslToken::RightParenthesis))
+		{
+			Parser.Scanner.SourceError(TEXT("')' expected!\n"));
+			return ParseResultError();
+		}
+
+		if (!Parser.Scanner.MatchToken(EHlslToken::Semicolon))
+		{
+			Parser.Scanner.SourceError(TEXT("';' expected!\n"));
+			return ParseResultError();
+		}
+
+		return EParseResult::Matched;
+	}
+
 	EParseResult ParseGlobalVariableDeclaration(FHlslParser& Parser, FLinearAllocator* Allocator, AST::FNode** OutDeclaration)
 	{
 		AST::FDeclaratorList* List = nullptr;
@@ -2208,6 +2262,7 @@ Done:
 				RulesStatements.Add(FRulePair(EHlslToken::Break, ParseBreakStatement));
 				RulesStatements.Add(FRulePair(EHlslToken::Continue, ParseContinueStatement));
 				RulesStatements.Add(FRulePair(EHlslToken::Invalid, ParseLocalDeclaration, true));
+				RulesStatements.Add(FRulePair(EHlslToken::StaticAssert, ParseStaticAssertStatement));
 				// Always try expressions last
 				RulesStatements.Add(FRulePair(EHlslToken::Invalid, ParseExpressionStatement));
 			}
