[Shaders] Forward CFLAG_WarningsAsErrors to DXC compiler in SPIR-V backends.

This affects Vulkan, Metal, GLES, and SPIR-V based platform backends.
Warnings will not be treated as errors for intermediate shader code, such as the one generated by the HLSL rewriter for dead-code removal.

#rnx

diff --git a/Engine/Source/Developer/Apple/MetalShaderFormat/Private/MetalCompileShaderSPIRV.cpp b/Engine/Source/Developer/Apple/MetalShaderFormat/Private/MetalCompileShaderSPIRV.cpp
--- a/Engine/Source/Developer/Apple/MetalShaderFormat/Private/MetalCompileShaderSPIRV.cpp
+++ b/Engine/Source/Developer/Apple/MetalShaderFormat/Private/MetalCompileShaderSPIRV.cpp
@@ -171,6 +171,7 @@ void FMetalCompileShaderSPIRV::DoCompileMetalShader(
 		CrossCompiler::FShaderConductorOptions Options;
 
 		Options.TargetEnvironment = CrossCompiler::FShaderConductorOptions::ETargetEnvironment::Vulkan_1_1;
+		Options.bWarningsAsErrors = Input.Environment.CompilerFlags.Contains(CFLAG_WarningsAsErrors);
 
 		// Enable HLSL 2021 if specified
 		if (Input.Environment.CompilerFlags.Contains(CFLAG_HLSL2021))
diff --git a/Engine/Source/Developer/ShaderCompilerCommon/Private/ShaderConductorContext.cpp b/Engine/Source/Developer/ShaderCompilerCommon/Private/ShaderConductorContext.cpp
--- a/Engine/Source/Developer/ShaderCompilerCommon/Private/ShaderConductorContext.cpp
+++ b/Engine/Source/Developer/ShaderCompilerCommon/Private/ShaderConductorContext.cpp
@@ -108,6 +108,7 @@ namespace CrossCompiler
 	{
 		FShaderConductorIntermediates()
 			: Stage(ShaderConductor::ShaderStage::NumShaderStages)
+			, bIsIntermediateCode(false)
 		{
 		}
 
@@ -123,6 +124,7 @@ namespace CrossCompiler
 		TArray<FAnsiString> CustomDxcArgs;
 		TArray<ANSICHAR const*> CustomDxcArgRefs;
 		TArray<ANSICHAR const*> DxcArgRefs;
+		bool bIsIntermediateCode; // Is the current shader source the result of intermediate compilation such as the DXC rewriter?
 	};
 
 	static void ConvertScSourceDesc(const FShaderConductorContext::FShaderConductorIntermediates& Intermediates, ShaderConductor::Compiler::SourceDesc& OutSourceDesc)
@@ -368,7 +370,7 @@ namespace CrossCompiler
 		}
 	}
 
-	static void AppendDxcArguments(const FShaderConductorOptions& InOptions, TArray<const ANSICHAR*>& DxcArguments, bool bGenerateSpirv = true)
+	static void AppendDxcArguments(const FShaderConductorOptions& InOptions, TArray<const ANSICHAR*>& DxcArguments, bool bIsIntermediateCode = false, bool bGenerateSpirv = true)
 	{
 		if(bGenerateSpirv)
 		{
@@ -404,6 +406,12 @@ namespace CrossCompiler
 			break;
 		}
 
+		// We only treat warnings as errors for input source code, not intermediate source since DXC rewriter might produce new warnings the shader authors don't have control over.
+		if (InOptions.bWarningsAsErrors && !bIsIntermediateCode)
+		{
+			DxcArguments.Add("-WX");
+		}
+
 		// Add additional DXC arguments that are not exposed by ShaderConductor API directly
 		if (!InOptions.bDisableScalarBlockLayout)
 		{
@@ -456,7 +464,7 @@ namespace CrossCompiler
 		}
 	}
 
-static void ConvertScOptions(FShaderConductorContext::FShaderConductorIntermediates& Intermediates, const FShaderConductorOptions& InOptions, ShaderConductor::Compiler::Options& OutOptions, bool bIgnoreCustomDxcArgs = false, bool bGenerateSpirv = true)
+	static void ConvertScOptions(FShaderConductorContext::FShaderConductorIntermediates& Intermediates, const FShaderConductorOptions& InOptions, ShaderConductor::Compiler::Options& OutOptions, bool bIgnoreCustomDxcArgs = false, bool bGenerateSpirv = true)
 	{
 		// Validate input shader model with respect to certain language features.
 		checkf(
@@ -483,8 +491,8 @@ static void ConvertScOptions(FShaderConductorContext::FShaderConductorIntermedia
 
 		DxcArgRefs.Empty();
 
-		AppendDxcArguments(InOptions, DxcArgRefs, bGenerateSpirv);
-		
+		AppendDxcArguments(InOptions, DxcArgRefs, Intermediates.bIsIntermediateCode, bGenerateSpirv);
+
 		if (!InOptions.SpirvCustomOptimizationPasses.IsEmpty())
 		{
 			Intermediates.InternalDxcArgs = FAnsiString::Printf("-Oconfig=%ls", SelectSpirvCustomOptimizationPasses(InOptions.SpirvCustomOptimizationPasses));
@@ -692,6 +700,9 @@ static void ConvertScOptions(FShaderConductorContext::FShaderConductorIntermedia
 			// Copy rewritten HLSL code into intermediate source code.
 			Intermediates->ShaderSource.CopyAnsi(ResultView);
 
+			// Mark the internal shader source as intermediate code
+			Intermediates->bIsIntermediateCode = true;
+
 			// If output source is specified, also convert to TCHAR string
 			if (OutSource != nullptr)
 			{
diff --git a/Engine/Source/Developer/ShaderCompilerCommon/Public/ShaderConductorContext.h b/Engine/Source/Developer/ShaderCompilerCommon/Public/ShaderConductorContext.h
--- a/Engine/Source/Developer/ShaderCompilerCommon/Public/ShaderConductorContext.h
+++ b/Engine/Source/Developer/ShaderCompilerCommon/Public/ShaderConductorContext.h
@@ -99,6 +99,9 @@ namespace CrossCompiler
 		bool bPreserveStorageInput = false;
         bool bForceStorageImageFormat = false;
 
+		/** Treat warnigns as errors. This adds '-WX' to the DXC arguments. See CFLAG_WarningsAsErrors. */
+		bool bWarningsAsErrors = false;
+
 		enum class ETargetEnvironment
 		{
 			Vulkan_1_0,
diff --git a/Engine/Source/Developer/ShaderFormatOpenGL/Private/OpenGLShaderCompiler.cpp b/Engine/Source/Developer/ShaderFormatOpenGL/Private/OpenGLShaderCompiler.cpp
--- a/Engine/Source/Developer/ShaderFormatOpenGL/Private/OpenGLShaderCompiler.cpp
+++ b/Engine/Source/Developer/ShaderFormatOpenGL/Private/OpenGLShaderCompiler.cpp
@@ -2883,6 +2883,7 @@ static bool CompileToGlslWithShaderConductor(
 	Options.bRemapAttributeLocations = true;
 	Options.bPreserveStorageInput = true;
     Options.bForceStorageImageFormat = true;
+	Options.bWarningsAsErrors = Input.Environment.CompilerFlags.Contains(CFLAG_WarningsAsErrors);
 	
 	// Enable HLSL 2021 if specified
 	if (Input.Environment.CompilerFlags.Contains(CFLAG_HLSL2021))
diff --git a/Engine/Source/Developer/VulkanShaderFormat/Private/VulkanShaderCompiler.cpp b/Engine/Source/Developer/VulkanShaderFormat/Private/VulkanShaderCompiler.cpp
--- a/Engine/Source/Developer/VulkanShaderFormat/Private/VulkanShaderCompiler.cpp
+++ b/Engine/Source/Developer/VulkanShaderFormat/Private/VulkanShaderCompiler.cpp
@@ -1924,6 +1924,7 @@ static bool CompileWithShaderConductor(
 	// Initialize compilation options for ShaderConductor
 	CrossCompiler::FShaderConductorOptions Options;
 	Options.TargetEnvironment = InternalState.MinimumTargetEnvironment;
+	Options.bWarningsAsErrors = Input.Environment.CompilerFlags.Contains(CFLAG_WarningsAsErrors);
 
 	// VK_EXT_scalar_block_layout is required by raytracing and by Nanite (so expect it to be present in SM6/Vulkan_1_3)
 	Options.bDisableScalarBlockLayout = !(InternalState.IsRayTracingShader() || InternalState.IsSM6());
