Handle line directives in ShaderMinifier ExtractBlock()

diff --git a/Engine/Source/Developer/ShaderCompilerCommon/Private/ShaderMinifier.cpp b/Engine/Source/Developer/ShaderCompilerCommon/Private/ShaderMinifier.cpp
--- a/Engine/Source/Developer/ShaderCompilerCommon/Private/ShaderMinifier.cpp
+++ b/Engine/Source/Developer/ShaderCompilerCommon/Private/ShaderMinifier.cpp
@@ -377,7 +377,7 @@ static FShaderSource::FViewType SkipUntilStr(FShaderSource::FViewType Haystack,
 	return SkipUntil(Haystack, [Needle](FShaderSource::FViewType  S) { return StartsWith(S, Needle); });
 }
 
-static FShaderSource::FViewType ExtractBlock(FShaderSource::FViewType Source, FShaderSource::CharType DelimBegin, FShaderSource::CharType DelimEnd)
+static FShaderSource::FViewType ExtractBlock(FShaderSource::FViewType Source, FShaderSource::CharType DelimBegin, FShaderSource::CharType DelimEnd, TArray<FShaderSource::FViewType>& OutLineDirectives)
 {
 	// TODO: handle comments
 	// TODO: handle #if 0 blocks
@@ -397,7 +397,7 @@ static FShaderSource::FViewType ExtractBlock(FShaderSource::FViewType Source, FS
 
 	EStatus Status = EStatus::Continue;
 
-	auto ProcessCharacter = [&Stack, &PosEnd, SourceData, DelimBegin, DelimEnd](int32 Cursor) -> EStatus
+	auto ProcessCharacter = [&Stack, &PosEnd, &OutLineDirectives, SourceData, SourceLen, DelimBegin, DelimEnd](int32 Cursor) -> EStatus
 	{
 		FShaderSource::CharType C = SourceData[Cursor];
 
@@ -421,14 +421,24 @@ static FShaderSource::FViewType ExtractBlock(FShaderSource::FViewType Source, FS
 				return EStatus::Finished;
 			}
 		}
+		else if (C == '#')
+		{
+			FShaderSource::FViewType Source(SourceData + Cursor, SourceLen - Cursor);
+			if (StartsWith(Source, SHADER_SOURCE_VIEWLITERAL("#line")))
+			{
+				FShaderSource::FViewType Remainder = SkipUntilNextLine(Source);
+				FShaderSource::FViewType Block = SubStrView(Source, 0, Source.Len() - Remainder.Len());
+				OutLineDirectives.Add(Block);
+			}
+		}
 
 		return EStatus::Continue;
 	};
 
 #if UE_SHADER_MINIFIER_SSE
 	const __m128i NeedleVec = FShaderSource::IsWide()
-		? _mm_setr_epi16(DelimBegin, DelimEnd, 0, 0, 0, 0, 0, 0)
-		: _mm_setr_epi8(DelimBegin, DelimEnd, 0, 0, 0, 0, 0, 0, 0, 0, 0, 0, 0, 0, 0, 0);
+		? _mm_setr_epi16(DelimBegin, DelimEnd, L'#', 0, 0, 0, 0, 0)
+		: _mm_setr_epi8(DelimBegin, DelimEnd, '#', 0, 0, 0, 0, 0, 0, 0, 0, 0, 0, 0, 0, 0);
 
 	while (Cursor < SourceLen && Status != EStatus::Finished)
 	{
@@ -1169,7 +1179,7 @@ static FParsedShader ParseShader(const FShaderSource& InSource, FDiagnostics& Ou
 		}
 		else if (C == '(')
 		{
-			Block = ExtractBlock(Source, '(', ')');
+			Block = ExtractBlock(Source, '(', ')', Result.LineDirectives);
 			BlockType = EBlockType::Args;
 
 			if (ArgsBlockIndex < 0)
@@ -1179,7 +1189,7 @@ static FParsedShader ParseShader(const FShaderSource& InSource, FDiagnostics& Ou
 		}
 		else if (C == '{')
 		{
-			Block = ExtractBlock(Source, '{', '}');
+			Block = ExtractBlock(Source, '{', '}', Result.LineDirectives);
 
 			if (BodyBlockIndex == INDEX_NONE && !bFoundAssignment)
 			{
@@ -1195,7 +1205,7 @@ static FParsedShader ParseShader(const FShaderSource& InSource, FDiagnostics& Ou
 		}
 		else if (C == '[')
 		{
-			Block = ExtractBlock(Source, '[', ']');
+			Block = ExtractBlock(Source, '[', ']', Result.LineDirectives);
 
 			if (bFoundIdentifier)
 			{
@@ -1208,7 +1218,7 @@ static FParsedShader ParseShader(const FShaderSource& InSource, FDiagnostics& Ou
 		}
 		else if (C == '<')
 		{
-			Block = ExtractBlock(Source, '<', '>');
+			Block = ExtractBlock(Source, '<', '>', Result.LineDirectives);
 			BlockType = EBlockType::TemplateArgs;
 
 			if (ChunkType == ECodeChunkType::CBuffer)
