Suppress false positive static analysis warning

diff --git a/Engine/Source/Developer/ShaderCompilerCommon/Private/ShaderMinifier.cpp b/Engine/Source/Developer/ShaderCompilerCommon/Private/ShaderMinifier.cpp
--- a/Engine/Source/Developer/ShaderCompilerCommon/Private/ShaderMinifier.cpp
+++ b/Engine/Source/Developer/ShaderCompilerCommon/Private/ShaderMinifier.cpp
@@ -2548,7 +2548,7 @@ float FunB(int Param)
 	return FunA() * (float)Param;
 }
 
-#line 1000 "MinifierTest.hlsl"
+#line 1000 "MinifierTest.hlsl" //-V011
 // Test comment 1
 void EmptyFunction(){}
 
