[shaders] further fixes for debug dump artifacts. if the shader format modifies source, ensure that is output as the <ShaderTypeName>.usf file as .bat files for invoking platform compiler directly use this version; in this case we output the directcompile-compatible version with a "_DirectCompile" suffix attached. For backends that don't modify source (i.e. don't set the ModifiedShaderSource field on the output struct) just output the directcompile-compatible version as the single <ShaderTypeName>.usf file (CL also includes cosmetic changes to suffix the dumped .usf files instead of prefixing, to avoid needing extra logic to combine prefixes when the backend already adds its own)

diff --git a/Engine/Source/Developer/ShaderCompilerCommon/Private/ShaderCompilerCommon.cpp b/Engine/Source/Developer/ShaderCompilerCommon/Private/ShaderCompilerCommon.cpp
--- a/Engine/Source/Developer/ShaderCompilerCommon/Private/ShaderCompilerCommon.cpp
+++ b/Engine/Source/Developer/ShaderCompilerCommon/Private/ShaderCompilerCommon.cpp
@@ -1374,7 +1374,7 @@ void CleanupUniformBufferCode(const FShaderCompilerEnvironment& Environment, FSh
 
 ///////////////////////////////////////////////////////////////////////////////////////////////////////////////////////////////////////
 
-FString CreateShaderCompilerWorkerDirectCommandLine(const FShaderCompilerInput& Input)
+FString CreateShaderCompilerWorkerDirectCommandLine(const FShaderCompilerInput& Input, const UE::ShaderCompilerCommon::FDebugShaderDataOptions& Options, const TCHAR* Suffix = nullptr)
 {
 	FString Text(TEXT("-directcompile -format="));
 	Text += Input.ShaderFormat.GetPlainNameString();
@@ -1418,7 +1418,7 @@ FString CreateShaderCompilerWorkerDirectCommandLine(const FShaderCompilerInput&
 	}
 
 	Text += TEXT(" ");
-	Text += Input.DumpDebugInfoPath / Input.GetSourceFilename();
+	Text += Options.GetDebugShaderPath(Input, Suffix);
 
 	// When we're running in directcompile mode, we don't to spam the crash reporter
 	Text += TEXT(" -nocrashreports");
@@ -1667,13 +1667,21 @@ void CompileOfflineMali(const FShaderCompilerInput& Input, FShaderCompilerOutput
 const FString GetDebugFileName(
 	const FShaderCompilerInput& Input, 
 	const UE::ShaderCompilerCommon::FDebugShaderDataOptions& Options, 
-	const TCHAR* BaseFilename)
+	const TCHAR* BaseFilename,
+	const TCHAR* Suffix = nullptr)
 {
 	TStringBuilder<512> PathBuilder;
 	const TCHAR* Prefix = (Options.FilenamePrefix && *Options.FilenamePrefix) ? Options.FilenamePrefix : TEXT("");
 	FStringView Filename = (BaseFilename && *BaseFilename) ? BaseFilename : Input.GetSourceFilenameView();
+	FStringView Ext = FPathViews::GetExtension(Filename, true);
+	FStringView FilenameNoExt = Filename.LeftChop(Ext.Len());
 	FPathViews::Append(PathBuilder, Input.DumpDebugInfoPath, Prefix);
-	PathBuilder << Filename;
+	PathBuilder << FilenameNoExt;
+	if (Suffix)
+	{
+		PathBuilder << Suffix;
+	}
+	PathBuilder << Ext;
 	return PathBuilder.ToString();
 }
 
@@ -1728,9 +1736,9 @@ namespace UE::ShaderCompilerCommon
 		return bSuccess;
 	}
 
-	FString FDebugShaderDataOptions::GetDebugShaderPath(const FShaderCompilerInput& Input) const
+	FString FDebugShaderDataOptions::GetDebugShaderPath(const FShaderCompilerInput& Input, const TCHAR* Suffix) const
 	{
-		return GetDebugFileName(Input, *this, OverrideBaseFilename);
+		return GetDebugFileName(Input, *this, OverrideBaseFilename, Suffix);
 	}
 
 	bool FBaseShaderFormat::PreprocessShader(
@@ -1749,19 +1757,19 @@ namespace UE::ShaderCompilerCommon
 		DumpExtendedDebugShaderData(Input, PreprocessOutput, Output);
 	}
 
-	void DumpDebugShaderData(const FShaderCompilerInput& Input, FStringView PreprocessedSource, const FDebugShaderDataOptions& Options)
+	void DumpDebugShaderData(const FShaderCompilerInput& Input, FStringView PreprocessedSource, const FDebugShaderDataOptions& Options, const TCHAR* Suffix)
 	{
 		if (!Input.DumpDebugInfoEnabled())
 		{
 			return;
 		}
 
-		FString Contents = UE::ShaderCompilerCommon::GetDebugShaderContents(Input, PreprocessedSource, Options);
-		FFileHelper::SaveStringToFile(Contents, *Options.GetDebugShaderPath(Input));
+		FString Contents = UE::ShaderCompilerCommon::GetDebugShaderContents(Input, PreprocessedSource, Options, Suffix);
+		FFileHelper::SaveStringToFile(Contents, *Options.GetDebugShaderPath(Input, Suffix));
 
 		if (EnumHasAnyFlags(Input.DebugInfoFlags, EShaderDebugInfoFlags::DirectCompileCommandLine) && !Options.bSourceOnly)
 		{
-			FFileHelper::SaveStringToFile(CreateShaderCompilerWorkerDirectCommandLine(Input), *GetDebugFileName(Input, Options, TEXT("DirectCompile.txt")));
+			FFileHelper::SaveStringToFile(CreateShaderCompilerWorkerDirectCommandLine(Input, Options, Suffix), *GetDebugFileName(Input, Options, TEXT("DirectCompile.txt")));
 		}
 	}
 
@@ -1773,32 +1781,22 @@ namespace UE::ShaderCompilerCommon
 	{
 		if (Input.bCachePreprocessed && EnumHasAnyFlags(Input.DebugInfoFlags, EShaderDebugInfoFlags::DetailedSource))
 		{
-			FDebugShaderDataOptions PrefixedOptions(Options);
-			uint32 SlackLen = Options.FilenamePrefix ? FCString::Strlen(Options.FilenamePrefix) : 0;
-			FString StrippedPrefix(TEXT("Stripped_"), SlackLen);
-			FString ModifiedPrefix(TEXT("CompileModified_"), SlackLen);
-			if (Options.FilenamePrefix)
-			{
-				StrippedPrefix += Options.FilenamePrefix;
-				ModifiedPrefix += Options.FilenamePrefix;
-			}
-			
-			PrefixedOptions.FilenamePrefix = *StrippedPrefix;
-			FFileHelper::SaveStringToFile(GetDebugShaderContents(Input, PreprocessOutput.GetSourceViewWide(), PrefixedOptions), *PrefixedOptions.GetDebugShaderPath(Input));
+			const TCHAR* StrippedSuffix = TEXT("_Stripped");
+			FFileHelper::SaveStringToFile(GetDebugShaderContents(Input, PreprocessOutput.GetSourceViewWide(), Options, StrippedSuffix), *Options.GetDebugShaderPath(Input, StrippedSuffix));
+		}
 
-			if (!Output.ModifiedShaderSource.IsEmpty())
-			{
-				// intentionally dumping this copy as-is rather than modifying via GetDebugShaderContents; this is not likely to be something that can be compiled in
-				// directcompile mode for debugging purposes, so we don't append the additional data that this requires.
-				PrefixedOptions.FilenamePrefix = *ModifiedPrefix;
-				FFileHelper::SaveStringToFile(Output.ModifiedShaderSource, *PrefixedOptions.GetDebugShaderPath(Input));
-			}
+		bool bHasModifiedSource = !Output.ModifiedShaderSource.IsEmpty();
+		if (bHasModifiedSource)
+		{
+			// If the compile step applies modifications to the source, output this as the "default" USF; it's not directcompile-compatible but backends
+			// which output compile batch files rely on this being the copy of the source that can be passed directly to the platform compiler.
+			FFileHelper::SaveStringToFile(Output.ModifiedShaderSource, *Options.GetDebugShaderPath(Input));
 		}
 
-		// Always output the preprocessed source (prior to stripping); this is guaranteed to work when passed to SCW in "directcompile" mode
-		// and is also more useful for error reporting (the stripped version of the source has line directives removed, and the error/warning remapping
-		// we perform is only done in the cooker process when processing completed jobs.
-		DumpDebugShaderData(Input, PreprocessOutput.GetUnstrippedSourceView(), Options);
+		// if no modifications to source are made in the compile step, output just the single usf which is the unstripped version compatible with launching
+		// SCW in directcompile mode (the stripped version is less useful for debugging via this mechanism, so is only output in "detailed source" mode)
+		// if modifications were made, this is output as an additional artifact, appending "_DirectCompile" to the path to indicate that it can be used as such.
+		DumpDebugShaderData(Input, PreprocessOutput.GetUnstrippedSourceView(), Options, bHasModifiedSource ? TEXT("_DirectCompile") : nullptr);
 
 		FFileHelper::SaveStringToFile(Output.OutputHash.ToString(), *GetDebugFileName(Input, Options, TEXT("OutputHash.txt")), FFileHelper::EEncodingOptions::ForceAnsi);
 
@@ -1876,7 +1874,7 @@ namespace UE::ShaderCompilerCommon
 		Env.SerializeCompilationDependencies(Ar);
 	}
 
-	FString GetDebugShaderContents(const FShaderCompilerInput& Input, FStringView PreprocessedSource, const FDebugShaderDataOptions& Options)
+	FString GetDebugShaderContents(const FShaderCompilerInput& Input, FStringView PreprocessedSource, const FDebugShaderDataOptions& Options, const TCHAR* Suffix)
 	{
 		// If preprocessed cache is enabled, debug dump occurs in the cook process rather than the workers, and
 		// in that case the env in Input.Environment has not been merged with the shared env. Do so here.
@@ -1903,7 +1901,7 @@ namespace UE::ShaderCompilerCommon
 		Contents += TEXT("\n");
 		Contents += SerializeEnvironmentToBase64(MergedEnvironment);
 		Contents += TEXT("/* DIRECT COMPILE\n");
-		Contents += CreateShaderCompilerWorkerDirectCommandLine(Input);
+		Contents += CreateShaderCompilerWorkerDirectCommandLine(Input, Options, Suffix);
 		Contents += TEXT("\nDIRECT COMPILE */\n");
 		if (!Input.DebugDescription.IsEmpty())
 		{
diff --git a/Engine/Source/Developer/ShaderCompilerCommon/Public/ShaderCompilerCommon.h b/Engine/Source/Developer/ShaderCompilerCommon/Public/ShaderCompilerCommon.h
--- a/Engine/Source/Developer/ShaderCompilerCommon/Public/ShaderCompilerCommon.h
+++ b/Engine/Source/Developer/ShaderCompilerCommon/Public/ShaderCompilerCommon.h
@@ -124,7 +124,7 @@ namespace UE::ShaderCompilerCommon
 		TArray<FAdditionalOutput> AdditionalOutputs;
 		bool bSourceOnly = false; // if true, will only output source .usf as directed and skip all other debug data artifacts
 
-		SHADERCOMPILERCOMMON_API FString GetDebugShaderPath(const FShaderCompilerInput& Input) const;
+		SHADERCOMPILERCOMMON_API FString GetDebugShaderPath(const FShaderCompilerInput& Input, const TCHAR* Suffix = nullptr) const;
 	};
 
 	/*
@@ -166,7 +166,7 @@ namespace UE::ShaderCompilerCommon
 	 * @param	PreprocessedSource The unmodified preprocessed source (used as input to the compilation)
 	 * @param	Options Options which can change behaviour of the debug dump; see above.
 	 */
-	extern SHADERCOMPILERCOMMON_API FString GetDebugShaderContents(const FShaderCompilerInput& Input, FStringView PreprocessedSource, const FDebugShaderDataOptions& Options = FDebugShaderDataOptions());
+	extern SHADERCOMPILERCOMMON_API FString GetDebugShaderContents(const FShaderCompilerInput& Input, FStringView PreprocessedSource, const FDebugShaderDataOptions& Options = FDebugShaderDataOptions(), const TCHAR* Suffix = nullptr);
 	
 	UE_DEPRECATED(5.4, "Use overload of GetDebugShaderContents accepting an FStringView")
 	inline FString GetDebugShaderContents(const FShaderCompilerInput& Input, const FString& PreprocessedSource, const FDebugShaderDataOptions& Options = FDebugShaderDataOptions())
