Fixing _Static_assert parsing.
- _Static_assert should be handled on the global scope.
- GetCurrentToken() doesn't advance to the next token, you have to PeekToken() and then MatchToken.
- No output was written on successful parsing of _Static_assert
- Removing debug assert when writing static_assert code back out.

diff --git a/Engine/Source/Developer/ShaderCompilerCommon/Private/HlslAST.cpp b/Engine/Source/Developer/ShaderCompilerCommon/Private/HlslAST.cpp
--- a/Engine/Source/Developer/ShaderCompilerCommon/Private/HlslAST.cpp
+++ b/Engine/Source/Developer/ShaderCompilerCommon/Private/HlslAST.cpp
@@ -535,8 +535,9 @@ namespace CrossCompiler
 			}
 		}
 
-		FStaticAssertStatement::FStaticAssertStatement(FLinearAllocator* InAllocator, const FSourceInfo& InInfo) :
+		FStaticAssertStatement::FStaticAssertStatement(FLinearAllocator* InAllocator, const FSourceInfo& InInfo, const FString& InKeyword) :
 			FNode(InAllocator, InInfo),
+			Keyword(InKeyword),
 			Condition(nullptr)
 		{
 		}
@@ -553,7 +554,6 @@ namespace CrossCompiler
 			checkf(Condition != nullptr, TEXT("Cannot write %s()-statement in HLSL without a conditional expression"), *Keyword);
 			Condition->Write(Writer);
 			Writer << TEXT(", \"") << *Message << TEXT("\");\n");
-			checkf(false, TEXT("TEST: %s()"), *Keyword);
 		}
 
 		FFunctionDefinition::FFunctionDefinition(FLinearAllocator* InAllocator, const FSourceInfo& InInfo) :
diff --git a/Engine/Source/Developer/ShaderCompilerCommon/Private/HlslAST.h b/Engine/Source/Developer/ShaderCompilerCommon/Private/HlslAST.h
--- a/Engine/Source/Developer/ShaderCompilerCommon/Private/HlslAST.h
+++ b/Engine/Source/Developer/ShaderCompilerCommon/Private/HlslAST.h
@@ -476,12 +476,12 @@ namespace CrossCompiler
 
 		struct FStaticAssertStatement : public FNode
 		{
-			FStaticAssertStatement(FLinearAllocator* InAllocator, const FSourceInfo& InInfo);
+			FStaticAssertStatement(FLinearAllocator* InAllocator, const FSourceInfo& InInfo, const FString& InKeyword);
 			~FStaticAssertStatement();
 
 			virtual void Write(FASTWriter& Writer) const override;
 
-			FString Keyword; // Either "_Static_assert" (DXC) or "static_assert" (C++11 like compilers)
+			const FString Keyword; // Either "_Static_assert" (DXC) or "static_assert" (C++11 like compilers)
 			FExpression* Condition;
 			FString Message;
 		};
diff --git a/Engine/Source/Developer/ShaderCompilerCommon/Private/HlslParser.cpp b/Engine/Source/Developer/ShaderCompilerCommon/Private/HlslParser.cpp
--- a/Engine/Source/Developer/ShaderCompilerCommon/Private/HlslParser.cpp
+++ b/Engine/Source/Developer/ShaderCompilerCommon/Private/HlslParser.cpp
@@ -1490,13 +1490,16 @@ Done:
 
 	EParseResult ParseStaticAssertStatement(FHlslParser& Parser, FLinearAllocator* Allocator, AST::FNode** OutDeclaration)
 	{
-		const FHlslToken* CurrentToken = Parser.Scanner.GetCurrentToken();
-		check(CurrentToken);
+		const FHlslToken* StaticAssertToken = Parser.Scanner.PeekToken();
+		if (!Parser.Scanner.MatchToken(EHlslToken::StaticAssert))
+		{
+			return EParseResult::NotMatched;
+		}
 
-		AST::FStaticAssertStatement* Statement = new(Allocator) AST::FStaticAssertStatement(Allocator, CurrentToken->SourceInfo);
+		check(StaticAssertToken);
 
 		// Store statement keyword to support _Static_assert (HLSL) and static_assert (C++11 like compilers)
-		Statement->Keyword = CurrentToken->String;
+		AST::FStaticAssertStatement* Statement = new(Allocator) AST::FStaticAssertStatement(Allocator, StaticAssertToken->SourceInfo, StaticAssertToken->String);
 
 		if (!Parser.Scanner.MatchToken(EHlslToken::LeftParenthesis))
 		{
@@ -1518,8 +1521,8 @@ Done:
 			return ParseResultError();
 		}
 
-		const FHlslToken* MessageToken = Parser.Scanner.GetCurrentToken();
-		if (MessageToken->Token != EHlslToken::StringConstant)
+		const FHlslToken* MessageToken = Parser.Scanner.PeekToken();
+		if (!Parser.Scanner.MatchToken(EHlslToken::StringConstant))
 		{
 			Parser.Scanner.SourceError(TEXT("')' expected!\n"));
 			return ParseResultError();
@@ -1539,6 +1542,8 @@ Done:
 			return ParseResultError();
 		}
 
+		*OutDeclaration = Statement;
+
 		return EParseResult::Matched;
 	}
 
@@ -2241,6 +2246,12 @@ Done:
 			return Result;
 		}
 
+		Result = ParseStaticAssertStatement(Parser, Allocator, OutNode);
+		if (Result == EParseResult::Error || Result == EParseResult::Matched)
+		{
+			return Result;
+		}
+
 		Parser.Scanner.SourceError(TEXT("Unable to match rule!"));
 		return ParseResultError();
 	}
@@ -2262,7 +2273,6 @@ Done:
 				RulesStatements.Add(FRulePair(EHlslToken::Break, ParseBreakStatement));
 				RulesStatements.Add(FRulePair(EHlslToken::Continue, ParseContinueStatement));
 				RulesStatements.Add(FRulePair(EHlslToken::Invalid, ParseLocalDeclaration, true));
-				RulesStatements.Add(FRulePair(EHlslToken::StaticAssert, ParseStaticAssertStatement));
 				// Always try expressions last
 				RulesStatements.Add(FRulePair(EHlslToken::Invalid, ParseExpressionStatement));
 			}
