[DXC] Temporarily disable DXC versioning to fix new DXC version with UBA.

"ImageGetDigestStream" is used to hash the dxcompiler.dll and dxil.dll files to detect outdated SCW builds. This function is buggy on POSIX/Wine, which is used in UBA and causes all shader complations through DXC to fail.

#rnx

diff --git a/Engine/Source/Developer/ShaderCompilerCommon/Private/DXCWrapper.cpp b/Engine/Source/Developer/ShaderCompilerCommon/Private/DXCWrapper.cpp
--- a/Engine/Source/Developer/ShaderCompilerCommon/Private/DXCWrapper.cpp
+++ b/Engine/Source/Developer/ShaderCompilerCommon/Private/DXCWrapper.cpp
@@ -28,6 +28,11 @@ BOOL __stdcall DigestHash(DIGEST_HANDLE Handle, PBYTE Data, DWORD Length)
 static uint64 GetLoadedModuleVersion(const TCHAR* ModuleName)
 {
 #if PLATFORM_WINDOWS
+	//@lh-todo - Temporarily disable DXC versioning mechanism until detoured "ImageGetDigestStream" in UBA is fixed (buggy on POSIX/Wine)
+#if 1
+	constexpr uint64 FixedVersionNo = 0x64da28316d6b9ea4;
+	return FixedVersionNo;
+#else
 	HMODULE ModuleDll = ::GetModuleHandleW(ModuleName);
 	if (ModuleDll == nullptr)
 	{
@@ -61,6 +66,7 @@ static uint64 GetLoadedModuleVersion(const TCHAR* ModuleName)
 	CloseHandle(DllHandle);
 
 	return FileHashBuilder.Finalize().Hash;
+#endif
 #else // PLATFORM_WINDOWS
 	return 0;
 #endif // PLATFORM_WINDOWS
