HlslParser improvements
- Adding support for templated function return types (Texture2d<float> etc)
- Adding support for ignoring empty statements (ie a ";" without code before it)

diff --git a/Engine/Source/Developer/ShaderCompilerCommon/Private/HlslExpressionParser.inl b/Engine/Source/Developer/ShaderCompilerCommon/Private/HlslExpressionParser.inl
--- a/Engine/Source/Developer/ShaderCompilerCommon/Private/HlslExpressionParser.inl
+++ b/Engine/Source/Developer/ShaderCompilerCommon/Private/HlslExpressionParser.inl
@@ -563,6 +563,55 @@ namespace CrossCompiler
 		return EParseResult::NotMatched;
 	}
 
+	EParseResult ParseFullType(FHlslScanner& Scanner, int32 TypeFlags, int32 TemplateTypeFlags, FSymbolScope* SymbolScope, FLinearAllocator* Allocator, AST::FFullySpecifiedType** OutFullySpecifiedType)
+	{
+		AST::FTypeSpecifier* TypeSpecifier = nullptr;
+		EParseResult Result = ParseGeneralType(Scanner, TypeFlags, SymbolScope, Allocator, &TypeSpecifier);
+		if (Result != EParseResult::Matched)
+		{
+			return Result;
+		}
+
+		// OutFullySpecifiedType could already be initialized
+		AST::FFullySpecifiedType* FullType = *OutFullySpecifiedType ? *OutFullySpecifiedType : new(Allocator) AST::FFullySpecifiedType(Allocator, TypeSpecifier->SourceInfo);
+		FullType->Specifier = TypeSpecifier;
+
+		if (Scanner.MatchToken(EHlslToken::Lower))
+		{
+			AST::FTypeSpecifier* ElementTypeSpecifier = nullptr;
+			EParseResult InnerResult = ParseGeneralType(Scanner, TemplateTypeFlags, SymbolScope, Allocator, &ElementTypeSpecifier);
+			if (InnerResult != EParseResult::Matched)
+			{
+				Scanner.SourceError(TEXT("Expected type!"));
+				return ParseResultError();
+			}
+
+			FullType->Specifier->InnerType = ElementTypeSpecifier->TypeName;
+
+			if (Scanner.MatchToken(EHlslToken::Comma))
+			{
+				auto* Integer = Scanner.GetCurrentToken();
+				if (!Scanner.MatchIntegerLiteral())
+				{
+					Scanner.SourceError(TEXT("Expected constant!"));
+					return ParseResultError();
+				}
+				FullType->Specifier->TextureMSNumSamples = FCString::Atoi(*Integer->String);
+			}
+
+			if (!Scanner.MatchToken(EHlslToken::Greater))
+			{
+				Scanner.SourceError(TEXT("Expected '>'!"));
+				return ParseResultError();
+			}
+		}
+
+		*OutFullySpecifiedType = FullType;
+
+		return EParseResult::Matched;
+	}
+
+
 	// Unary!(Unary-(Unary+())) would have ! as Top, and + as Inner
 	EParseResult MatchUnaryOperator(FHlslScanner& Scanner, /*FInfo& Info,*/ FSymbolScope* SymbolScope, FLinearAllocator* Allocator, AST::FExpression** OuterExpression, AST::FExpression** InnerExpression)
 	{
diff --git a/Engine/Source/Developer/ShaderCompilerCommon/Private/HlslParser.cpp b/Engine/Source/Developer/ShaderCompilerCommon/Private/HlslParser.cpp
--- a/Engine/Source/Developer/ShaderCompilerCommon/Private/HlslParser.cpp
+++ b/Engine/Source/Developer/ShaderCompilerCommon/Private/HlslParser.cpp
@@ -291,45 +291,11 @@ namespace CrossCompiler
 
 	EParseResult ParseTextureOrBufferSimpleDeclaration(FHlslScanner& Scanner, FSymbolScope* SymbolScope, bool bMultiple, bool bInitializer, bool bRegister, FLinearAllocator* Allocator, AST::FDeclaratorList** OutDeclaratorList)
 	{
-		auto OriginalToken = Scanner.GetCurrentTokenIndex();
-		const auto* Token = Scanner.GetCurrentToken();
-		auto* FullType = (*OutDeclaratorList)->Type;
-		if (ParseGeneralType(Scanner, ETF_SAMPLER_TEXTURE_BUFFER, nullptr, Allocator, &FullType->Specifier) == EParseResult::Matched)
+		const uint32 OriginalToken = Scanner.GetCurrentTokenIndex();
+		const FHlslToken* Token = Scanner.GetCurrentToken();
+		AST::FFullySpecifiedType* FullType = (*OutDeclaratorList)->Type;
+		if (ParseFullType(Scanner, ETF_SAMPLER_TEXTURE_BUFFER, ETF_BUILTIN_NUMERIC | ETF_USER_TYPES | ETF_UNORM, SymbolScope, Allocator, &FullType) == EParseResult::Matched)
 		{
-			if (Scanner.MatchToken(EHlslToken::Lower))
-			{
-				AST::FTypeSpecifier* ElementTypeSpecifier = nullptr;
-				auto Result = ParseGeneralType(Scanner, ETF_BUILTIN_NUMERIC | ETF_USER_TYPES | ETF_UNORM, SymbolScope, Allocator, &ElementTypeSpecifier);
-				if (Result != EParseResult::Matched)
-				{
-					Scanner.SourceError(TEXT("Expected type!"));
-					return ParseResultError();
-				}
-
-				FullType->Specifier->InnerType = ElementTypeSpecifier->TypeName;
-
-				if (Scanner.MatchToken(EHlslToken::Comma))
-				{
-					auto* Integer = Scanner.GetCurrentToken();
-					if (!Scanner.MatchIntegerLiteral())
-					{
-						Scanner.SourceError(TEXT("Expected constant!"));
-						return ParseResultError();
-					}
-					FullType->Specifier->TextureMSNumSamples = FCString::Atoi(*Integer->String);
-				}
-
-				if (!Scanner.MatchToken(EHlslToken::Greater))
-				{
-					Scanner.SourceError(TEXT("Expected '>'!"));
-					return ParseResultError();
-				}
-			}
-			else
-			{
-				//TypeSpecifier->InnerName = "float4";
-			}
-
 			do
 			{
 				// Handle 'Sampler2D Sampler'
@@ -1200,11 +1166,18 @@ namespace CrossCompiler
 		return nullptr;
 	}
 
+	EParseResult ParseFunctionReturnType(FHlslParser& Parser, FLinearAllocator* Allocator, AST::FFullySpecifiedType** OutFullySpecifiedType)
+	{
+		const int32 TypeFlags = ETF_BUILTIN_NUMERIC | ETF_SAMPLER_TEXTURE_BUFFER | ETF_USER_TYPES | ETF_ERROR_IF_NOT_USER_TYPE | ETF_VOID;
+		const int32 TemplateTypeFlags = ETF_BUILTIN_NUMERIC | ETF_SAMPLER_TEXTURE_BUFFER | ETF_USER_TYPES | ETF_UNORM;
+		return ParseFullType(Parser.Scanner, TypeFlags, TemplateTypeFlags, Parser.CurrentScope, Allocator, OutFullySpecifiedType);
+	}
+
 	EParseResult ParseFunctionDeclarator(FHlslParser& Parser, FLinearAllocator* Allocator, AST::FFunction** OutFunction)
 	{
 		uint32 OriginalToken = Parser.Scanner.GetCurrentTokenIndex();
-		AST::FTypeSpecifier* TypeSpecifier = nullptr;
-		EParseResult Result = ParseGeneralType(Parser.Scanner, ETF_BUILTIN_NUMERIC | ETF_SAMPLER_TEXTURE_BUFFER | ETF_USER_TYPES | ETF_ERROR_IF_NOT_USER_TYPE | ETF_VOID, Parser.CurrentScope, Allocator, &TypeSpecifier);
+		AST::FFullySpecifiedType* ReturnType = nullptr;
+		EParseResult Result = ParseFunctionReturnType(Parser, Allocator, &ReturnType);
 		if (Result == EParseResult::NotMatched)
 		{
 			Parser.Scanner.SetCurrentTokenIndex(OriginalToken);
@@ -1282,8 +1255,7 @@ namespace CrossCompiler
 		{
 			Function->ScopeIdentifier = Allocator->Strdup(ScopeIdentifier->String);
 		}
-		Function->ReturnType = new(Allocator) AST::FFullySpecifiedType(Allocator, TypeSpecifier->SourceInfo);
-		Function->ReturnType->Specifier = TypeSpecifier;
+		Function->ReturnType = ReturnType;
 
 		if (Parser.Scanner.MatchToken(EHlslToken::Void))
 		{
@@ -2252,6 +2224,12 @@ Done:
 			return Result;
 		}
 
+		// Ignore semicolons with no code before them.
+		if (Parser.Scanner.MatchToken(EHlslToken::Semicolon))
+		{
+			return EParseResult::Ignore;
+		}
+
 		Parser.Scanner.SourceError(TEXT("Unable to match rule!"));
 		return ParseResultError();
 	}
@@ -2376,6 +2354,10 @@ Done:
 					bSuccess = false;
 					break;
 				}
+				else if (Result == EParseResult::Ignore)
+				{
+					continue;
+				}
 				else
 				{
 					check(Result == EParseResult::Matched);
diff --git a/Engine/Source/Developer/ShaderCompilerCommon/Private/HlslParser.h b/Engine/Source/Developer/ShaderCompilerCommon/Private/HlslParser.h
--- a/Engine/Source/Developer/ShaderCompilerCommon/Private/HlslParser.h
+++ b/Engine/Source/Developer/ShaderCompilerCommon/Private/HlslParser.h
@@ -24,6 +24,7 @@ namespace CrossCompiler
 	{
 		Matched,
 		NotMatched,
+		Ignore,
 		Error,
 	};
 
