ShaderMinifier fix to prevent presence numthreads attribute pulling in all other functions with numthreads attribute.
This regression happened when we added extraction of identifiers from attributes.
That is specifically needed for extracting name from a NodeID attribute.
To fix we now treat NodeID attributes as their own block type.

diff --git a/Engine/Source/Developer/ShaderCompilerCommon/Private/ShaderMinifier.cpp b/Engine/Source/Developer/ShaderCompilerCommon/Private/ShaderMinifier.cpp
--- a/Engine/Source/Developer/ShaderCompilerCommon/Private/ShaderMinifier.cpp
+++ b/Engine/Source/Developer/ShaderCompilerCommon/Private/ShaderMinifier.cpp
@@ -493,6 +493,7 @@ enum class EBlockType : uint8 {
 	NamespaceDelimiter, // e.g. :: in an identifier like Foo::bar
 	PtrOrRef, // e.g. '*' or '&' as part of the type
 	OperatorName, // overladed operator, e.g. '+', '+=', etc.
+	NodeId, // e.g. [NodeID(Foo, 0)]
 };
 
 struct FCodeBlock
@@ -1211,6 +1212,10 @@ static FParsedShader ParseShader(const FShaderSource& InSource, FDiagnostics& Ou
 			{
 				BlockType = EBlockType::Subscript;
 			}
+			else if (Block.Contains(SHADER_SOURCE_LITERAL("NodeID")))
+			{
+				BlockType = EBlockType::NodeId;
+			}
 			else
 			{
 				BlockType = EBlockType::Attribute;
@@ -1777,7 +1782,7 @@ static FShaderSource::FStringType MinifyShader(const FParsedShader& Parsed, TCon
 				continue;
 			}
 
-			if (Chunk.Type == ECodeChunkType::Function && Block.Type != EBlockType::Name && Block.Type != EBlockType::Attribute)
+			if (Chunk.Type == ECodeChunkType::Function && Block.Type != EBlockType::Name && Block.Type != EBlockType::NodeId)
 			{
 				continue;
 			}
@@ -1801,7 +1806,7 @@ static FShaderSource::FStringType MinifyShader(const FParsedShader& Parsed, TCon
 
 			bool bExtractIdentifiers = (Chunk.Type == ECodeChunkType::CBuffer && Block.Type == EBlockType::Body);
 			bExtractIdentifiers |= (Chunk.Type == ECodeChunkType::Enum && Block.Type == EBlockType::Body);
-			bExtractIdentifiers |= (Chunk.Type == ECodeChunkType::Function && Block.Type == EBlockType::Attribute);
+			bExtractIdentifiers |= (Chunk.Type == ECodeChunkType::Function && Block.Type == EBlockType::NodeId);
 
 			if (bExtractIdentifiers)
 			{
@@ -2484,7 +2489,7 @@ bool FShaderMinifierParserTest::RunTest(const FString& Parameters)
 			TestEqual(TEXT("ParseShader: workgraph node, chunk type"), P.Chunks[0].Type, ECodeChunkType::Function);
 			if (TestEqual(TEXT("ParseShader: workgraph node: chunk 0: num blocks"), P.Chunks[0].Blocks.Num(), 6))
 			{
-				TestEqual(TEXT("ParseShader: workgraph node: chunk 0: block 1: block type"), P.Chunks[0].Blocks[1].Type, EBlockType::Attribute);
+				TestEqual(TEXT("ParseShader: workgraph node: chunk 0: block 1: block type"), P.Chunks[0].Blocks[1].Type, EBlockType::NodeId);
 			}
 		}
 	}
@@ -2651,6 +2656,11 @@ void MainCS(
 	OutputBuffer[0] = A + B + D;
 }
 
+[numthreads(1,1,1)]
+void UnreferencedEntryPoint()
+{
+}
+
 [Shader("node")]
 void WorkgraphNode()
 {
@@ -2752,6 +2762,7 @@ void WorkgraphNodeArray_1()
 		TestFalse(TEXT("MinifyShader: MainCS: contains EEnumUnused"), MinifiedParsed.Source.Contains(SHADER_SOURCE_LITERAL("EEnumUnused")));
 		TestFalse(TEXT("MinifyShader: MainCS: contains ENUM_UNUSED_PART_1"), MinifiedParsed.Source.Contains(SHADER_SOURCE_LITERAL("ENUM_UNUSED_PART_1")));
 		TestFalse(TEXT("MinifyShader: MainCS: contains ENUM_UNUSED_PART_2"), MinifiedParsed.Source.Contains(SHADER_SOURCE_LITERAL("ENUM_UNUSED_PART_2")));
+		TestFalse(TEXT("MinifyShader: MainCS: contains UnreferencedEntryPoint"), MinifiedParsed.Source.Contains(SHADER_SOURCE_LITERAL("UnreferencedEntryPoint")));
 	}
 
 	int32 NumErrors = ExecutionInfo.GetErrorTotal();
