[Shaders] Remove last remaining use case of DXC rewriter in MetalCompileShaderMSC and deprecate RewriteHlsl.

#rnx

diff --git a/Engine/Source/Developer/Apple/MetalShaderFormat/Private/MetalCompileShaderMSC.cpp b/Engine/Source/Developer/Apple/MetalShaderFormat/Private/MetalCompileShaderMSC.cpp
--- a/Engine/Source/Developer/Apple/MetalShaderFormat/Private/MetalCompileShaderMSC.cpp
+++ b/Engine/Source/Developer/Apple/MetalShaderFormat/Private/MetalCompileShaderMSC.cpp
@@ -508,14 +508,6 @@ void FMetalCompileShaderMSC::DoCompileMetalShader(
 		// Load shader source into compiler context
 		CompilerContext.LoadSource(PreprocessedShader, Input.VirtualSourceFilePath, Input.EntryPointName, Frequency);
 
-		// Rewrite HLSL source code to remove unused global resources and variables
-		Options.bRemoveUnusedGlobals = true;
-		if (!CompilerContext.RewriteHlsl(Options, &PreprocessedShader))
-		{
-			CompilerContext.FlushErrors(Output.Errors);
-		}
-		Options.bRemoveUnusedGlobals = false;
-
 		// Convert shader source to ANSI string
 		std::string SourceData(CompilerContext.GetSourceString(), static_cast<size_t>(CompilerContext.GetSourceLength()));
 
diff --git a/Engine/Source/Developer/ShaderCompilerCommon/Private/ShaderConductorContext.cpp b/Engine/Source/Developer/ShaderCompilerCommon/Private/ShaderConductorContext.cpp
--- a/Engine/Source/Developer/ShaderCompilerCommon/Private/ShaderConductorContext.cpp
+++ b/Engine/Source/Developer/ShaderCompilerCommon/Private/ShaderConductorContext.cpp
@@ -678,44 +678,7 @@ namespace CrossCompiler
 
 	bool FShaderConductorContext::RewriteHlsl(const FShaderConductorOptions& Options, FString* OutSource)
 	{
-		// Convert descriptors for ShaderConductor interface
-		ShaderConductor::Compiler::SourceDesc ScSourceDesc;
-		ConvertScSourceDesc(*Intermediates, ScSourceDesc);
-
-		ShaderConductor::Compiler::Options ScOptions;
-		constexpr bool bIgnoreExtraDxcArgs = true;
-		ConvertScOptions(*Intermediates, Options, ScOptions, bIgnoreExtraDxcArgs);
-
-		// Rewrite HLSL with wrapper function to catch exceptions from ShaderConductor
-		bool bSucceeded = false;
-		ShaderConductor::Compiler::ResultDesc ResultDesc;
-		ScRewriteWrapper(ScSourceDesc, ScOptions, ResultDesc);
-
-		if (!ResultDesc.hasError && ResultDesc.target.Size() > 1)
-		{
-			// Note: We don't want to include the '\0' included in the result string (thanks to DxcCreateBlob), hence the -1
-			check(reinterpret_cast<const ANSICHAR*>(ResultDesc.target.Data())[ResultDesc.target.Size() - 1] == '\0');
-			FAnsiStringView ResultView(reinterpret_cast<const ANSICHAR*>(ResultDesc.target.Data()), ResultDesc.target.Size() - 1);
-
-			// Copy rewritten HLSL code into intermediate source code.
-			Intermediates->ShaderSource.CopyAnsi(ResultView);
-
-			// Mark the internal shader source as intermediate code
-			Intermediates->bIsIntermediateCode = true;
-
-			// If output source is specified, also convert to TCHAR string
-			if (OutSource != nullptr)
-			{
-				OutSource->Empty();
-				OutSource->Append(ResultView);
-			}
-			bSucceeded = true;
-		}
-
-		// Append compile error and warning to output reports
-		ConvertScCompileErrors(ResultDesc.errorWarningMsg, Errors);
-
-		return bSucceeded;
+		return false; // DXC rewriter is no longer supported and this function has been deprecated since UE5.5
 	}
 
     bool FShaderConductorContext::CompileHlslToDxil(const FShaderConductorOptions& Options, TArray<uint32>& OutDxil)
diff --git a/Engine/Source/Developer/ShaderCompilerCommon/Public/ShaderConductorContext.h b/Engine/Source/Developer/ShaderCompilerCommon/Public/ShaderConductorContext.h
--- a/Engine/Source/Developer/ShaderCompilerCommon/Public/ShaderConductorContext.h
+++ b/Engine/Source/Developer/ShaderCompilerCommon/Public/ShaderConductorContext.h
@@ -211,6 +211,7 @@ namespace CrossCompiler
 
 		/** Rewrites the specified HLSL shader source code. This allows to reduce the HLSL code by removing unused global resources for instance.
 		This will update the internally loaded source (see 'LoadSource'), so the output parameter 'OutSource' is optional. */
+		UE_DEPRECATED(5.5, "DXC rewriter has been deprecated since UE5.5. FShaderConductorContext::RewriteHlsl will be removed in future versions.")
 		bool RewriteHlsl(const FShaderConductorOptions& Options, FString* OutSource = nullptr);
 
         /** Compiles the specified HLSL shader source code to DXIL. */
