HLSL Parser: Add support for C99 style _Pragma statements

This alternative style of pragma is useful for emitting pragmas from macros.

Since these pragmas take strings as their arguments, improve the support for parsing string constants to deal with escaped characters.

Add some examples of this to the unit tests

diff --git a/Engine/Source/Developer/ShaderCompilerCommon/Private/HlslAST.cpp b/Engine/Source/Developer/ShaderCompilerCommon/Private/HlslAST.cpp
--- a/Engine/Source/Developer/ShaderCompilerCommon/Private/HlslAST.cpp
+++ b/Engine/Source/Developer/ShaderCompilerCommon/Private/HlslAST.cpp
@@ -556,6 +556,21 @@ namespace CrossCompiler
 			Writer << TEXT(", \"") << *Message << TEXT("\");\n");
 		}
 
+		FC99PragmaStatement::FC99PragmaStatement(FLinearAllocator* InAllocator, const FSourceInfo& InInfo, const FString& InArgument) :
+			FNode(InAllocator, InInfo),
+			Argument(InArgument)
+		{
+		}
+
+		FC99PragmaStatement::~FC99PragmaStatement()
+		{
+		}
+
+		void FC99PragmaStatement::Write(FASTWriter& Writer) const
+		{
+			Writer << TEXT("_Pragma(\"") << *Argument << TEXT("\")\n");
+		}
+
 		FFunctionDefinition::FFunctionDefinition(FLinearAllocator* InAllocator, const FSourceInfo& InInfo) :
 			FNode(InAllocator, InInfo),
 			Prototype(nullptr),
diff --git a/Engine/Source/Developer/ShaderCompilerCommon/Private/HlslAST.h b/Engine/Source/Developer/ShaderCompilerCommon/Private/HlslAST.h
--- a/Engine/Source/Developer/ShaderCompilerCommon/Private/HlslAST.h
+++ b/Engine/Source/Developer/ShaderCompilerCommon/Private/HlslAST.h
@@ -486,6 +486,16 @@ namespace CrossCompiler
 			FString Message;
 		};
 
+		struct FC99PragmaStatement : public FNode
+		{
+			FC99PragmaStatement(FLinearAllocator* InAllocator, const FSourceInfo& InInfo, const FString& InArgument);
+			~FC99PragmaStatement();
+
+			virtual void Write(FASTWriter& Writer) const override;
+
+			FString Argument;
+		};
+
 		struct FSemanticSpecifier : public FNode
 		{
 			enum class ESpecType
diff --git a/Engine/Source/Developer/ShaderCompilerCommon/Private/HlslLexer.cpp b/Engine/Source/Developer/ShaderCompilerCommon/Private/HlslLexer.cpp
--- a/Engine/Source/Developer/ShaderCompilerCommon/Private/HlslLexer.cpp
+++ b/Engine/Source/Developer/ShaderCompilerCommon/Private/HlslLexer.cpp
@@ -442,6 +442,7 @@ namespace CrossCompiler
 				InsertToken(TEXT("operator"), EHlslToken::Operator);
 				InsertToken(TEXT("_Static_assert"), EHlslToken::StaticAssert); // HLSL2021 adopted C11 '_Static_assert'-statements
 				InsertToken(TEXT("static_assert"), EHlslToken::StaticAssert); // Some shader compilers support C++11 'static_assert'-statements
+				InsertToken(TEXT("_Pragma"), EHlslToken::C99Pragma); // C99/C++11 style pragma (distinct from #pragma because argument is a string constant)
 			}
 		} GStaticInitializer;
 	}
@@ -770,9 +771,22 @@ namespace CrossCompiler
 			OutString = TEXT("");
 			while (Peek() != '"')
 			{
-				OutString += Peek();
-				//@todo-rco: Check for \"
-				//@todo-rco: Check for EOL
+				auto Char = Peek();
+				OutString += Char;
+				if (Char == 0)
+				{
+					return false; // ill-formed string (EOL inside quote)
+				}
+				if (Char == '\\') // escaped character, we must have at least one more
+				{
+					++Current;
+					Char = Peek();
+					if (Char == 0)
+					{
+						return false; // EOL
+					}
+					OutString += Char; //@todo-rco: Should we validate we have escaped a valid character?
+				}
 				++Current;
 			}
 
diff --git a/Engine/Source/Developer/ShaderCompilerCommon/Private/HlslLexer.h b/Engine/Source/Developer/ShaderCompilerCommon/Private/HlslLexer.h
--- a/Engine/Source/Developer/ShaderCompilerCommon/Private/HlslLexer.h
+++ b/Engine/Source/Developer/ShaderCompilerCommon/Private/HlslLexer.h
@@ -283,6 +283,7 @@ namespace CrossCompiler
 		BoolConstant,
 		StringConstant,	// C-style "string"
 		StaticAssert, // _Static_assert (HLSL), static_assert (C++11)
+		C99Pragma, // _Pragma (C99/C++11)
 	};
 
 	enum class ELiteralType
diff --git a/Engine/Source/Developer/ShaderCompilerCommon/Private/HlslParser.cpp b/Engine/Source/Developer/ShaderCompilerCommon/Private/HlslParser.cpp
--- a/Engine/Source/Developer/ShaderCompilerCommon/Private/HlslParser.cpp
+++ b/Engine/Source/Developer/ShaderCompilerCommon/Private/HlslParser.cpp
@@ -1519,6 +1519,43 @@ Done:
 		return EParseResult::Matched;
 	}
 
+	EParseResult ParseC99PragmaStatementBody(FHlslParser& Parser, FLinearAllocator* Allocator, AST::FNode** OutDeclaration)
+	{
+		if (!Parser.Scanner.MatchToken(EHlslToken::LeftParenthesis))
+		{
+			Parser.Scanner.SourceError(TEXT("'(' expected!\n"));
+			return ParseResultError();
+		}
+
+		const FHlslToken* PragmaArgumentToken = Parser.Scanner.PeekToken();
+		if (!Parser.Scanner.MatchToken(EHlslToken::StringConstant))
+		{
+			Parser.Scanner.SourceError(TEXT("string constant expected!\n"));
+			return ParseResultError();
+		}
+		// Store the string inside the pragma so we can round-trip it properly
+		AST::FC99PragmaStatement* Statement = new(Allocator) AST::FC99PragmaStatement(Allocator, PragmaArgumentToken->SourceInfo, PragmaArgumentToken->String);
+
+		if (!Parser.Scanner.MatchToken(EHlslToken::RightParenthesis))
+		{
+			Parser.Scanner.SourceError(TEXT("')' expected!\n"));
+			return ParseResultError();
+		}
+
+		*OutDeclaration = Statement;
+
+		return EParseResult::Matched;
+	}
+
+	EParseResult ParseC99PragmaStatement(FHlslParser& Parser, FLinearAllocator* Allocator, AST::FNode** OutDeclaration)
+	{
+		if (!Parser.Scanner.MatchToken(EHlslToken::C99Pragma))
+		{
+			return EParseResult::NotMatched;
+		}
+		return ParseC99PragmaStatementBody(Parser, Allocator, OutDeclaration);
+	}
+
 	EParseResult ParseGlobalVariableDeclaration(FHlslParser& Parser, FLinearAllocator* Allocator, AST::FNode** OutDeclaration)
 	{
 		AST::FDeclaratorList* List = nullptr;
@@ -2224,6 +2261,12 @@ Done:
 			return Result;
 		}
 
+		Result = ParseC99PragmaStatement(Parser, Allocator, OutNode);
+		if (Result == EParseResult::Error || Result == EParseResult::Matched)
+		{
+			return Result;
+		}
+
 		// Ignore semicolons with no code before them.
 		if (Parser.Scanner.MatchToken(EHlslToken::Semicolon))
 		{
@@ -2240,6 +2283,7 @@ Done:
 		{
 			FStaticInitializer()
 			{
+				RulesStatements.Add(FRulePair(EHlslToken::C99Pragma, ParseC99PragmaStatementBody));
 				RulesStatements.Add(FRulePair(EHlslToken::LeftBrace, ParseStatementBlock));
 				RulesStatements.Add(FRulePair(EHlslToken::Return, ParseReturnStatement));
 				RulesStatements.Add(FRulePair(EHlslToken::Do, ParseDoStatement));
diff --git a/Engine/Source/Developer/ShaderCompilerCommon/Private/ShaderMinifier.cpp b/Engine/Source/Developer/ShaderCompilerCommon/Private/ShaderMinifier.cpp
--- a/Engine/Source/Developer/ShaderCompilerCommon/Private/ShaderMinifier.cpp
+++ b/Engine/Source/Developer/ShaderCompilerCommon/Private/ShaderMinifier.cpp
@@ -2531,6 +2531,12 @@ uint UnreferencedFunction()
 	return GUnreferencedParameter;
 }
 
+_Pragma("")
+
+_Pragma("dxc diagnostic push")
+_Pragma("dxc diagnostic ignored \"-Wall\"") _Pragma("dxc diagnostic ignored \"-Wconversion\"")
+_Pragma("dxc diagnostic pop")
+
 typedef Texture2D<float4> UnreferencedTypedef;
 
 #define COMPILER_DEFINITION_TEST 123
