[Shaders] Remove unused state in FShaderConductorContext and ignore unknown attributes in shader backends.

This adds "-Wno-unknown-attributes" to DXC backends and platform specific compiler argument to their respective backend.

#rnx

diff --git a/Engine/Source/Developer/ShaderCompilerCommon/Private/ShaderConductorContext.cpp b/Engine/Source/Developer/ShaderCompilerCommon/Private/ShaderConductorContext.cpp
--- a/Engine/Source/Developer/ShaderCompilerCommon/Private/ShaderConductorContext.cpp
+++ b/Engine/Source/Developer/ShaderCompilerCommon/Private/ShaderConductorContext.cpp
@@ -108,7 +108,6 @@ namespace CrossCompiler
 	{
 		FShaderConductorIntermediates()
 			: Stage(ShaderConductor::ShaderStage::NumShaderStages)
-			, bIsIntermediateCode(false)
 		{
 		}
 
@@ -124,7 +123,6 @@ namespace CrossCompiler
 		TArray<FAnsiString> CustomDxcArgs;
 		TArray<ANSICHAR const*> CustomDxcArgRefs;
 		TArray<ANSICHAR const*> DxcArgRefs;
-		bool bIsIntermediateCode; // Is the current shader source the result of intermediate compilation such as the DXC rewriter?
 	};
 
 	static void ConvertScSourceDesc(const FShaderConductorContext::FShaderConductorIntermediates& Intermediates, ShaderConductor::Compiler::SourceDesc& OutSourceDesc)
@@ -370,7 +368,7 @@ namespace CrossCompiler
 		}
 	}
 
-	static void AppendDxcArguments(const FShaderConductorOptions& InOptions, TArray<const ANSICHAR*>& DxcArguments, bool bIsIntermediateCode = false, bool bGenerateSpirv = true)
+	static void AppendDxcArguments(const FShaderConductorOptions& InOptions, TArray<const ANSICHAR*>& DxcArguments, bool bGenerateSpirv = true)
 	{
 		if(bGenerateSpirv)
 		{
@@ -406,8 +404,11 @@ namespace CrossCompiler
 			break;
 		}
 
+		// Ignore unknwon attributes as UE uses custom attributes for intermediate source transformation
+		DxcArguments.Add("-Wno-unknown-attributes");
+
 		// We only treat warnings as errors for input source code, not intermediate source since DXC rewriter might produce new warnings the shader authors don't have control over.
-		if (InOptions.bWarningsAsErrors && !bIsIntermediateCode)
+		if (InOptions.bWarningsAsErrors)
 		{
 			DxcArguments.Add("-WX");
 		}
@@ -491,7 +492,7 @@ namespace CrossCompiler
 
 		DxcArgRefs.Empty();
 
-		AppendDxcArguments(InOptions, DxcArgRefs, Intermediates.bIsIntermediateCode, bGenerateSpirv);
+		AppendDxcArguments(InOptions, DxcArgRefs, bGenerateSpirv);
 
 		if (!InOptions.SpirvCustomOptimizationPasses.IsEmpty())
 		{
