- Fix for AddUsedInputMembers having inverted source / destination for arrays

diff --git a/Engine/Shaders/Public/ShaderVersion.ush b/Engine/Shaders/Public/ShaderVersion.ush
--- a/Engine/Shaders/Public/ShaderVersion.ush
+++ b/Engine/Shaders/Public/ShaderVersion.ush
@@ -3,4 +3,4 @@
 // Each time the console command is executed it generates a new GUID. As this file is included
 // in Platform.ush (which should be included in any shader) it allows to invalidate the shader DDC.
 // 
-#pragma message("UESHADERMETADATA_VERSION F5BD4964-B08E-431F-B710-DC98AD2EA1FC")
+#pragma message("UESHADERMETADATA_VERSION 4BA01FFC-127B-AE42-AAA5-259B8F406FCB")
diff --git a/Engine/Source/Developer/ShaderCompilerCommon/Private/HlslUtils.cpp b/Engine/Source/Developer/ShaderCompilerCommon/Private/HlslUtils.cpp
--- a/Engine/Source/Developer/ShaderCompilerCommon/Private/HlslUtils.cpp
+++ b/Engine/Source/Developer/ShaderCompilerCommon/Private/HlslUtils.cpp
@@ -1290,7 +1290,8 @@ struct FRemoveUnusedInputs : FRemoveAlgorithm
 								NewDeclaratorList->Declarations.Add(MemberDeclaration);
 								DestStruct->Members.Add(NewDeclaratorList);
 
-								CopyMember(MemberDeclaration, DestPrefix, SourcePrefix, BodyContext.PreInstructions);
+								// Source and Dest are swapped as we are copying from the optimized (dest) structure into the original (source) structure
+								CopyMember(MemberDeclaration, SourcePrefix, DestPrefix, BodyContext.PreInstructions);
 							}
 
 							for (uint32 Index = 0; Index < ArrayLength; ++Index)
@@ -1330,6 +1331,7 @@ struct FRemoveUnusedInputs : FRemoveAlgorithm
 							NewDeclaratorList->Declarations.Add(MemberDeclaration);
 							DestStruct->Members.Add(NewDeclaratorList);
 
+							// Source and Dest are swapped as we are copying from the optimized (dest) structure into the original (source) structure
 							CopyMember(MemberDeclaration, SourcePrefix, DestPrefix, BodyContext.PreInstructions);
 						}
 						else
