[Backout] - CL33681066
Original CL Desc
-----------------------------------------------------------------
[Backout] - CL33671805
Original CL Desc
-----------------------------------------------------------------
- Fix for AddUsedInputMembers having inverted source / destination for arrays

diff --git a/Engine/Shaders/Public/ShaderVersion.ush b/Engine/Shaders/Public/ShaderVersion.ush
--- a/Engine/Shaders/Public/ShaderVersion.ush
+++ b/Engine/Shaders/Public/ShaderVersion.ush
@@ -3,4 +3,4 @@
 // Each time the console command is executed it generates a new GUID. As this file is included
 // in Platform.ush (which should be included in any shader) it allows to invalidate the shader DDC.
 // 
-#pragma message("UESHADERMETADATA_VERSION F5BD4964-B08E-431F-B710-DC98AD2EA1FC")
+#pragma message("UESHADERMETADATA_VERSION 4BA01FFC-127B-AE42-AAA5-259B8F406FCB")
diff --git a/Engine/Source/Developer/ShaderCompilerCommon/Private/HlslUtils.cpp b/Engine/Source/Developer/ShaderCompilerCommon/Private/HlslUtils.cpp
--- a/Engine/Source/Developer/ShaderCompilerCommon/Private/HlslUtils.cpp
+++ b/Engine/Source/Developer/ShaderCompilerCommon/Private/HlslUtils.cpp
@@ -1290,7 +1290,8 @@ struct FRemoveUnusedInputs : FRemoveAlgorithm
 								NewDeclaratorList->Declarations.Add(MemberDeclaration);
 								DestStruct->Members.Add(NewDeclaratorList);
 
-								CopyMember(MemberDeclaration, DestPrefix, SourcePrefix, BodyContext.PreInstructions);
+								// Source and Dest are swapped as we are copying from the optimized (dest) structure into the original (source) structure
+								CopyMember(MemberDeclaration, SourcePrefix, DestPrefix, BodyContext.PreInstructions);
 							}
 
 							for (uint32 Index = 0; Index < ArrayLength; ++Index)
@@ -1330,6 +1331,7 @@ struct FRemoveUnusedInputs : FRemoveAlgorithm
 							NewDeclaratorList->Declarations.Add(MemberDeclaration);
 							DestStruct->Members.Add(NewDeclaratorList);
 
+							// Source and Dest are swapped as we are copying from the optimized (dest) structure into the original (source) structure
 							CopyMember(MemberDeclaration, SourcePrefix, DestPrefix, BodyContext.PreInstructions);
 						}
 						else
diff --git a/Engine/Source/Runtime/Core/Private/UObject/FortniteMainBranchObjectVersion.cpp b/Engine/Source/Runtime/Core/Private/UObject/FortniteMainBranchObjectVersion.cpp
--- a/Engine/Source/Runtime/Core/Private/UObject/FortniteMainBranchObjectVersion.cpp
+++ b/Engine/Source/Runtime/Core/Private/UObject/FortniteMainBranchObjectVersion.cpp
@@ -11,7 +11,7 @@ TMap<FGuid, FGuid> FFortniteMainBranchObjectVersion::GetSystemGuids()
 	SystemGuids.Add(DevGuids.MATERIALSHADERMAP_DERIVEDDATA_VER, FGuid("BB56D95CD4FD48E38C4CD7869B44B987"));
 	SystemGuids.Add(DevGuids.NANITE_DERIVEDDATA_VER, FGuid("D28C003F825F4C178F698938FB0E47F7"));
 	SystemGuids.Add(DevGuids.NIAGARASHADERMAP_DERIVEDDATA_VER, FGuid("C417A50271B3427D98785388266E3B66"));
-	SystemGuids.Add(DevGuids.Niagara_LatestScriptCompileVersion, FGuid("C2285DEBF1B14D48BC67C27764A13F16"));
+	SystemGuids.Add(DevGuids.Niagara_LatestScriptCompileVersion, FGuid("E486231DA967C64B95D4E73131056FA0"));
 	SystemGuids.Add(DevGuids.SkeletalMeshDerivedDataVersion, FGuid("B9AEAA3EE7AC4FCFB3860E8F3E06DD75"));
 	SystemGuids.Add(DevGuids.STATICMESH_DERIVEDDATA_VER, FGuid("543843B103794E0E9BA4BE80FB602F79"));
 	SystemGuids.Add(DevGuids.MaterialTranslationDDCVersion, FGuid("F43DE2FD43D353A9ABE260AD4CE9F3D4"));
