Reducing allocations when using FShaderParameterMap
- TCHAR* arguments were being used to find elements in a TMap<FString,..> which creates temporary memory allocations.
- We can use FStringView with FindByHash/RemoveByHash to prevent temporary allocations.
- Adding FShaderParameterMap::FindAndRemoveParameterAllocation as a shortcut for Find+Remove (which Uniform Buffer member handling does).
- Fixing a few locations that were going FString->TCHAR*->FString.

diff --git a/Engine/Source/Developer/ShaderCompilerCommon/Private/ShaderCompilerCommon.cpp b/Engine/Source/Developer/ShaderCompilerCommon/Private/ShaderCompilerCommon.cpp
--- a/Engine/Source/Developer/ShaderCompilerCommon/Private/ShaderCompilerCommon.cpp
+++ b/Engine/Source/Developer/ShaderCompilerCommon/Private/ShaderCompilerCommon.cpp
@@ -119,7 +119,7 @@ bool BuildResourceTableMapping(
 		const FString& Name = Entry.UniformBufferMemberName;
 
 		// If the shaders uses this member (eg View_PerlinNoise3DTexture)...
-		if (TOptional<FParameterAllocation> Allocation = ParameterMap.FindParameterAllocation(Name))
+		if (TOptional<FParameterAllocation> Allocation = ParameterMap.FindAndRemoveParameterAllocation(Name))
 		{
 			const EShaderParameterType ParameterType = Allocation->Type;
 			const bool bBindlessParameter = IsParameterBindless(ParameterType);
@@ -127,17 +127,18 @@ bool BuildResourceTableMapping(
 			// Force bindless "indices" to zero since they're not needed in SetResourcesFromTables
 			const uint16 BaseIndex = bBindlessParameter ? 0 : Allocation->BaseIndex;
 
-			ParameterMap.RemoveParameterAllocation(*Name);
-
 			uint16 UniformBufferIndex = INDEX_NONE;
-			uint16 UBBaseIndex, UBSize;
 
 			// Add the UB itself as a parameter if not there
-			FString UniformBufferName(Entry.GetUniformBufferName());
-			if (!ParameterMap.FindParameterAllocation(*UniformBufferName, UniformBufferIndex, UBBaseIndex, UBSize))
+			FStringView UniformBufferName = Entry.GetUniformBufferName();
+			if (TOptional<FParameterAllocation> UniformBufferParameter = ParameterMap.FindParameterAllocation(UniformBufferName))
+			{
+				UniformBufferIndex = UniformBufferParameter->BufferIndex;
+			}
+			else
 			{
 				UniformBufferIndex = UsedUniformBufferSlots.FindAndSetFirstZeroBit();
-				ParameterMap.AddParameterAllocation(*UniformBufferName, UniformBufferIndex,0,0,EShaderParameterType::UniformBuffer);
+				ParameterMap.AddParameterAllocation(UniformBufferName, UniformBufferIndex, 0, 0, EShaderParameterType::UniformBuffer);
 			}
 
 			// Mark used UB index
@@ -221,7 +222,7 @@ void CullGlobalUniformBuffers(const TMap<FString, FUniformBufferEntry>& UniformB
 				continue;
 			}
 
-			ParameterMap.RemoveParameterAllocation(*Name);
+			ParameterMap.RemoveParameterAllocation(Name);
 		}
 	}
 }
@@ -403,7 +404,7 @@ bool UE::ShaderCompilerCommon::ValidatePackedResourceCounts(FShaderCompilerOutpu
 	{
 		auto GetAllResourcesOfType = [&](EShaderParameterType InType)
 		{
-			const TArray<FString> AllNames = Output.ParameterMap.GetAllParameterNamesOfType(InType);
+			const TArray<FStringView> AllNames = Output.ParameterMap.GetAllParameterNamesOfType(InType);
 			if (AllNames.IsEmpty())
 			{
 				return FString();
@@ -571,7 +572,7 @@ void HandleReflectedGlobalConstantBufferMember(
 	const EShaderParameterType ParameterType = FShaderParameterParser::ParseAndRemoveBindlessParameterPrefix(MemberName);
 
 	Output.ParameterMap.AddParameterAllocation(
-		*MemberName,
+		MemberName,
 		ConstantBufferIndex,
 		ReflectionOffset,
 		ReflectionSize,
@@ -592,7 +593,7 @@ void HandleReflectedUniformBufferConstantBufferMember(
 	if (ParameterType != EShaderParameterType::LooseData)
 	{
 		Output.ParameterMap.AddParameterAllocation(
-			*MemberName,
+			MemberName,
 			UniformBufferSlot,
 			ReflectionOffset,
 			1,
@@ -645,7 +646,7 @@ void HandleReflectedUniformBuffer(
 	FString AdjustedUniformBufferName(UE::ShaderCompilerCommon::RemoveConstantBufferPrefix(UniformBufferName));
 
 	CompilerOutput.ParameterMap.AddParameterAllocation(
-		*AdjustedUniformBufferName,
+		AdjustedUniformBufferName,
 		ReflectionSlot,
 		BaseIndex,
 		BufferSize,
@@ -662,7 +663,7 @@ void HandleReflectedShaderResource(
 )
 {
 	CompilerOutput.ParameterMap.AddParameterAllocation(
-		*ResourceName,
+		ResourceName,
 		BindOffset,
 		ReflectionSlot,
 		BindCount,
@@ -728,7 +729,7 @@ void HandleReflectedShaderUAV(
 )
 {
 	CompilerOutput.ParameterMap.AddParameterAllocation(
-		*UAVName,
+		UAVName,
 		BindOffset,
 		ReflectionSlot,
 		BindCount,
@@ -745,7 +746,7 @@ void HandleReflectedShaderSampler(
 )
 {
 	CompilerOutput.ParameterMap.AddParameterAllocation(
-		*SamplerName,
+		SamplerName,
 		BindOffset,
 		ReflectionSlot,
 		BindCount,
diff --git a/Engine/Source/Developer/VulkanShaderFormat/Private/VulkanShaderCompiler.cpp b/Engine/Source/Developer/VulkanShaderFormat/Private/VulkanShaderCompiler.cpp
--- a/Engine/Source/Developer/VulkanShaderFormat/Private/VulkanShaderCompiler.cpp
+++ b/Engine/Source/Developer/VulkanShaderFormat/Private/VulkanShaderCompiler.cpp
@@ -541,8 +541,8 @@ static void AddUniformBuffer(
 		checkf(!UBInfo.bOnlyHasResources, TEXT("UBName = %s, BindingIndex = %d"), *UBName, (int32)BindingIndex);
 	}
 
-	InOutParameterMap.RemoveParameterAllocation(*UBName);
-	InOutParameterMap.AddParameterAllocation(*UBName, HeaderUBIndex, 0, 1, EShaderParameterType::UniformBuffer);
+	InOutParameterMap.RemoveParameterAllocation(UBName);
+	InOutParameterMap.AddParameterAllocation(UBName, HeaderUBIndex, 0, 1, EShaderParameterType::UniformBuffer);
 }
 
 static int32 DoAddGlobal(const FString& Name, FVulkanShaderHeader& OutHeader, TArray<FString>& OutGlobalNames)
@@ -595,7 +595,7 @@ static void PrepareGlobals(const VulkanShaderCompilerSerializedOutput& Serialize
 	for (int32 ParameterIndex = 0; ParameterIndex < ParameterNames.Num(); ++ParameterIndex)
 	{
 		const FString& ParameterName = ParameterNames[ParameterIndex];
-		TOptional<FParameterAllocation> ParameterAllocation = ParameterMap.FindParameterAllocation(*ParameterName);
+		TOptional<FParameterAllocation> ParameterAllocation = ParameterMap.FindParameterAllocation(ParameterName);
 		checkf(ParameterAllocation.IsSet(), TEXT("PrepareGlobals failed to find resource ParameterName=%s"), *ParameterName);
 
 		if (ParameterAllocation->Type == EShaderParameterType::UniformBuffer)
@@ -622,7 +622,7 @@ static void PrepareGlobals(const VulkanShaderCompilerSerializedOutput& Serialize
 	for (int32 ParameterIndex = 0; ParameterIndex < ParameterNames.Num(); ++ParameterIndex)
 	{
 		const FString& ParameterName = ParameterNames[ParameterIndex];
-		TOptional<FParameterAllocation> ParameterAllocation = ParameterMap.FindParameterAllocation(*ParameterName);
+		TOptional<FParameterAllocation> ParameterAllocation = ParameterMap.FindParameterAllocation(ParameterName);
 		checkf(ParameterAllocation.IsSet(), TEXT("PrepareGlobals failed to find resource ParameterName=%s"), *ParameterName);
 			
 		if (ParameterAllocation->Type == EShaderParameterType::UniformBuffer)
@@ -689,8 +689,8 @@ static void ConvertToHeader(
 		{
 			const int32 HeaderGlobalIndex = AddGlobal(ParameterName, ParameterAllocation->BaseIndex, SerializedOutput, GlobalNames);
 
-			InOutParameterMap.RemoveParameterAllocation(*ParameterName);
-			InOutParameterMap.AddParameterAllocation(*ParameterName, 0, HeaderGlobalIndex, ParameterAllocation->Size, ParameterAllocation->Type);
+			InOutParameterMap.RemoveParameterAllocation(ParameterName);
+			InOutParameterMap.AddParameterAllocation(ParameterName, 0, HeaderGlobalIndex, ParameterAllocation->Size, ParameterAllocation->Type);
 		}
 	}
 
@@ -872,7 +872,7 @@ static void BuildShaderOutput(
 	for (int32 Index = NewParameters.Num() - 1; Index >= 0; --Index)
 	{
 		uint16 OutIndex, OutBase, OutSize;
-		const bool bFound = ShaderOutput.ParameterMap.FindParameterAllocation(*NewParameters[Index], OutIndex, OutBase, OutSize);
+		const bool bFound = ShaderOutput.ParameterMap.FindParameterAllocation(NewParameters[Index], OutIndex, OutBase, OutSize);
 		ensure(bFound);
 		NumParams = FMath::Max((uint16)(OutIndex + 1), NumParams);
 		if (OriginalParameters.Contains(NewParameters[Index]))
@@ -1969,7 +1969,7 @@ static void UpdateBindlessUBs(const FVulkanShaderCompilerInternalState& Internal
 #endif
 
 			const int32 UBIndex = SerializedOutput.Header.UniformBuffers.Num() - 1;
-			Output.ParameterMap.AddParameterAllocation(*CBName, UBIndex, 0, 1, EShaderParameterType::UniformBuffer);
+			Output.ParameterMap.AddParameterAllocation(CBName, UBIndex, 0, 1, EShaderParameterType::UniformBuffer);
 		}
 	}
 }
diff --git a/Engine/Source/Runtime/RenderCore/Private/ShaderCore.cpp b/Engine/Source/Runtime/RenderCore/Private/ShaderCore.cpp
--- a/Engine/Source/Runtime/RenderCore/Private/ShaderCore.cpp
+++ b/Engine/Source/Runtime/RenderCore/Private/ShaderCore.cpp
@@ -700,9 +700,9 @@ bool AllowGlobalShaderLoad()
 
 }
 
-TOptional<FParameterAllocation> FShaderParameterMap::FindParameterAllocation(const FString& ParameterName) const
+TOptional<FParameterAllocation> FShaderParameterMap::FindParameterAllocation(FStringView ParameterName) const
 {
-	if (const FParameterAllocation* Allocation = ParameterMap.Find(ParameterName))
+	if (const FParameterAllocation* Allocation = ParameterMap.FindByHash(GetTypeHash(ParameterName), ParameterName))
 	{
 		if (Allocation->bBound)
 		{
@@ -718,7 +718,18 @@ TOptional<FParameterAllocation> FShaderParameterMap::FindParameterAllocation(con
 	return TOptional<FParameterAllocation>();
 }
 
-bool FShaderParameterMap::FindParameterAllocation(const TCHAR* ParameterName, uint16& OutBufferIndex, uint16& OutBaseIndex, uint16& OutSize) const
+TOptional<FParameterAllocation> FShaderParameterMap::FindAndRemoveParameterAllocation(FStringView ParameterName)
+{
+	FParameterAllocation Result;
+	if (ParameterMap.RemoveAndCopyValueByHash(GetTypeHash(ParameterName), ParameterName, Result))
+	{
+		return TOptional<FParameterAllocation>(Result);
+	}
+
+	return TOptional<FParameterAllocation>();
+}
+
+bool FShaderParameterMap::FindParameterAllocation(FStringView ParameterName, uint16& OutBufferIndex, uint16& OutBaseIndex, uint16& OutSize) const
 {
 	if (TOptional<FParameterAllocation> Allocation = FindParameterAllocation(ParameterName))
 	{
@@ -732,25 +743,25 @@ bool FShaderParameterMap::FindParameterAllocation(const TCHAR* ParameterName, ui
 	return false;
 }
 
-bool FShaderParameterMap::ContainsParameterAllocation(const TCHAR* ParameterName) const
+bool FShaderParameterMap::ContainsParameterAllocation(FStringView ParameterName) const
 {
-	return ParameterMap.Find(ParameterName) != NULL;
+	return ParameterMap.FindByHash(GetTypeHash(ParameterName), ParameterName) != nullptr;
 }
 
-void FShaderParameterMap::AddParameterAllocation(const TCHAR* ParameterName,uint16 BufferIndex,uint16 BaseIndex,uint16 Size,EShaderParameterType ParameterType)
+void FShaderParameterMap::AddParameterAllocation(FStringView ParameterName,uint16 BufferIndex,uint16 BaseIndex,uint16 Size,EShaderParameterType ParameterType)
 {
 	check(ParameterType < EShaderParameterType::Num);
-	ParameterMap.Add(ParameterName, FParameterAllocation(BufferIndex, BaseIndex, Size, ParameterType));
+	ParameterMap.Emplace(ParameterName, FParameterAllocation(BufferIndex, BaseIndex, Size, ParameterType));
 }
 
-void FShaderParameterMap::RemoveParameterAllocation(const TCHAR* ParameterName)
+void FShaderParameterMap::RemoveParameterAllocation(FStringView ParameterName)
 {
-	ParameterMap.Remove(ParameterName);
+	ParameterMap.RemoveByHash(GetTypeHash(ParameterName), ParameterName);
 }
 
-TArray<FString> FShaderParameterMap::GetAllParameterNamesOfType(EShaderParameterType InType) const
+TArray<FStringView> FShaderParameterMap::GetAllParameterNamesOfType(EShaderParameterType InType) const
 {
-	TArray<FString> Result;
+	TArray<FStringView> Result;
 	for (const TMap<FString, FParameterAllocation>::ElementType& Parameter : ParameterMap)
 	{
 		if (Parameter.Value.Type == InType)
diff --git a/Engine/Source/Runtime/RenderCore/Private/ShaderParameterStruct.cpp b/Engine/Source/Runtime/RenderCore/Private/ShaderParameterStruct.cpp
--- a/Engine/Source/Runtime/RenderCore/Private/ShaderParameterStruct.cpp
+++ b/Engine/Source/Runtime/RenderCore/Private/ShaderParameterStruct.cpp
@@ -366,11 +366,10 @@ void FShaderParameterBindings::BindForRootShaderParameters(const FShader* Shader
 	// Binds the uniform buffer that contains the root shader parameters.
 	{
 		const TCHAR* ShaderBindingName = FShaderParametersMetadata::kRootUniformBufferBindingName;
-		uint16 BufferIndex, BaseIndex, BoundSize;
-		if (ParametersMap.FindParameterAllocation(ShaderBindingName, BufferIndex, BaseIndex, BoundSize))
+		if (TOptional<FParameterAllocation> Parameter = ParametersMap.FindParameterAllocation(ShaderBindingName))
 		{
 			BindingContext.ShaderGlobalScopeBindings.Add(ShaderBindingName, ShaderBindingName);
-			RootParameterBufferIndex = BufferIndex;
+			RootParameterBufferIndex = Parameter->BufferIndex;
 		}
 		else
 		{
diff --git a/Engine/Source/Runtime/RenderCore/Private/ShaderParameters.cpp b/Engine/Source/Runtime/RenderCore/Private/ShaderParameters.cpp
--- a/Engine/Source/Runtime/RenderCore/Private/ShaderParameters.cpp
+++ b/Engine/Source/Runtime/RenderCore/Private/ShaderParameters.cpp
@@ -102,18 +102,20 @@ void FShaderUniformBufferParameter::Bind(const FShaderParameterMap& ParameterMap
 	uint16 UnusedBaseIndex = 0;
 	uint16 UnusedNumBytes = 0;
 
-	if(!ParameterMap.FindParameterAllocation(ParameterName,BaseIndex,UnusedBaseIndex,UnusedNumBytes))
+	if (TOptional<FParameterAllocation> Parameter = ParameterMap.FindParameterAllocation(ParameterName))
+	{
+		// NOTE: the name difference is intentional (and confusing)
+		BaseIndex = Parameter->BufferIndex;
+		checkf(IsBound(), TEXT("UniformBuffer Parameter '%s' was not bound with a valid index. This can cause instability at runtime."), ParameterName);
+	}
+	else
 	{
 		BaseIndex = 0xffff;
-		if(Flags == SPF_Mandatory)
+		if (Flags == SPF_Mandatory)
 		{
 			FailureToBindNonOptionalParameter(TEXT("shader resource parameter"), ParameterName);
 		}
 	}
-	else
-	{
-		check(IsBound());
-	}
 }
 
 #if WITH_EDITOR
diff --git a/Engine/Source/Runtime/RenderCore/Public/ShaderCore.h b/Engine/Source/Runtime/RenderCore/Public/ShaderCore.h
--- a/Engine/Source/Runtime/RenderCore/Public/ShaderCore.h
+++ b/Engine/Source/Runtime/RenderCore/Public/ShaderCore.h
@@ -296,14 +296,15 @@ public:
 	FShaderParameterMap()
 	{}
 
-	RENDERCORE_API TOptional<FParameterAllocation> FindParameterAllocation(const FString& ParameterName) const;
-	RENDERCORE_API bool FindParameterAllocation(const TCHAR* ParameterName,uint16& OutBufferIndex,uint16& OutBaseIndex,uint16& OutSize) const;
-	RENDERCORE_API bool ContainsParameterAllocation(const TCHAR* ParameterName) const;
-	RENDERCORE_API void AddParameterAllocation(const TCHAR* ParameterName,uint16 BufferIndex,uint16 BaseIndex,uint16 Size,EShaderParameterType ParameterType);
-	RENDERCORE_API void RemoveParameterAllocation(const TCHAR* ParameterName);
+	RENDERCORE_API TOptional<FParameterAllocation> FindParameterAllocation(FStringView ParameterName) const;
+	RENDERCORE_API TOptional<FParameterAllocation> FindAndRemoveParameterAllocation(FStringView ParameterName);
+	RENDERCORE_API bool FindParameterAllocation(FStringView ParameterName, uint16& OutBufferIndex, uint16& OutBaseIndex, uint16& OutSize) const;
+	RENDERCORE_API bool ContainsParameterAllocation(FStringView ParameterName) const;
+	RENDERCORE_API void AddParameterAllocation(FStringView ParameterName, uint16 BufferIndex, uint16 BaseIndex, uint16 Size, EShaderParameterType ParameterType);
+	RENDERCORE_API void RemoveParameterAllocation(FStringView ParameterName);
 
 	/** Returns an array of all parameters with the given type. */
-	RENDERCORE_API TArray<FString> GetAllParameterNamesOfType(EShaderParameterType InType) const;
+	RENDERCORE_API TArray<FStringView> GetAllParameterNamesOfType(EShaderParameterType InType) const;
 
 	/** Returns a count of all parameters of the given type. */
 	RENDERCORE_API uint32 CountParametersOfType(EShaderParameterType InType) const;
@@ -323,12 +324,15 @@ public:
 		return Ar;
 	}
 
-	inline void GetAllParameterNames(TArray<FString>& OutNames) const
+	void GetAllParameterNames(TArray<FString>& OutNames) const
 	{
 		ParameterMap.GenerateKeyArray(OutNames);
 	}
 
-	inline const TMap<FString, FParameterAllocation>& GetParameterMap() const { return ParameterMap; }
+	const TMap<FString, FParameterAllocation>& GetParameterMap() const
+	{
+		return ParameterMap;
+	}
 
 	TMap<FString,FParameterAllocation> ParameterMap;
 };
