Uniform Buffer improvements
- Moving various UB booleans into a flags enum.
- UB booleans could not be reasonably deprecated without incurring memory overhead, so this will break custom code that uses them.
- Adding UB flag to force the shader compilers to generate reflection for the UB members which are normally excluded from reflection.
- Adding UB flag that tells MeshCommands that a UB will be bound during pass drawing and that it doesn't need to be set via MDCs.
- New flags are not used in this CL, they are prerequisites for subsequent, larger changes.

diff --git a/Engine/Shaders/Public/ShaderVersion.ush b/Engine/Shaders/Public/ShaderVersion.ush
--- a/Engine/Shaders/Public/ShaderVersion.ush
+++ b/Engine/Shaders/Public/ShaderVersion.ush
@@ -3,4 +3,4 @@
 // Each time the console command is executed it generates a new GUID. As this file is included
 // in Platform.ush (which should be included in any shader) it allows to invalidate the shader DDC.
 // 
-#pragma message("UESHADERMETADATA_VERSION 2b7d999b-f44f-496d-9bdf-50dcd21236bb")
+#pragma message("UESHADERMETADATA_VERSION 810370AA-2780-4601-9316-5186F4550D88")
diff --git a/Engine/Source/Developer/Apple/MetalShaderFormat/Private/MetalCompileShaderMSC.cpp b/Engine/Source/Developer/Apple/MetalShaderFormat/Private/MetalCompileShaderMSC.cpp
--- a/Engine/Source/Developer/Apple/MetalShaderFormat/Private/MetalCompileShaderMSC.cpp
+++ b/Engine/Source/Developer/Apple/MetalShaderFormat/Private/MetalCompileShaderMSC.cpp
@@ -272,28 +272,35 @@ static void ProcessReflection(ID3D12ShaderReflection* ShaderReflection, const ui
 				D3D12_SHADER_BUFFER_DESC CBDesc;
 				ConstantBuffer->GetDesc(&CBDesc);
 
-				for (uint32 ConstantIndex = 0; ConstantIndex < CBDesc.Variables; ConstantIndex++)
+				const FString UniformBufferName(BindDesc.Name);
+				const EUniformBufferMemberReflectionReason Reason = ShouldReflectUniformBufferMembers(Input, UniformBufferName);
+				if (Reason != EUniformBufferMemberReflectionReason::None)
 				{
-					ID3D12ShaderReflectionVariable* Variable = ConstantBuffer->GetVariableByIndex(ConstantIndex);
+					for (uint32 ConstantIndex = 0; ConstantIndex < CBDesc.Variables; ConstantIndex++)
+					{
+						ID3D12ShaderReflectionVariable* Variable = ConstantBuffer->GetVariableByIndex(ConstantIndex);
 
-					D3D12_SHADER_VARIABLE_DESC VariableDesc;
-					Variable->GetDesc(&VariableDesc);
+						D3D12_SHADER_VARIABLE_DESC VariableDesc;
+						Variable->GetDesc(&VariableDesc);
 
-					if (VariableDesc.uFlags & D3D_SVF_USED)
-					{
-						HandleReflectedUniformBufferConstantBufferMember(
-							BindIndex,
-							FString(VariableDesc.Name),
-							VariableDesc.StartOffset,
-							VariableDesc.Size,
-							Output
-						);
+						if (VariableDesc.uFlags & D3D_SVF_USED)
+						{
+							HandleReflectedUniformBufferConstantBufferMember(
+								Reason,
+								UniformBufferName,
+								BindIndex,
+								FString(VariableDesc.Name),
+								VariableDesc.StartOffset,
+								VariableDesc.Size,
+								Output
+							);
+						}
 					}
 				}
 				
 				// Regular uniform buffer - we only care about the binding index
-				CCHeaderWriter.WriteUniformBlock(ANSI_TO_TCHAR(BindDesc.Name), BindIndex);
-				HandleReflectedUniformBuffer(ANSI_TO_TCHAR(BindDesc.Name), BindIndex, Output);
+				CCHeaderWriter.WriteUniformBlock(*UniformBufferName, BindIndex);
+				HandleReflectedUniformBuffer(UniformBufferName, BindIndex, Output);
 			}
 			NumCBVs = FMath::Max(NumCBVs, BindIndex + BindDesc.BindCount);
 		}
diff --git a/Engine/Source/Developer/ShaderCompilerCommon/Private/ShaderCompilerCommon.cpp b/Engine/Source/Developer/ShaderCompilerCommon/Private/ShaderCompilerCommon.cpp
--- a/Engine/Source/Developer/ShaderCompilerCommon/Private/ShaderCompilerCommon.cpp
+++ b/Engine/Source/Developer/ShaderCompilerCommon/Private/ShaderCompilerCommon.cpp
@@ -99,6 +99,15 @@ void UE::ShaderCompilerCommon::BuildShaderResourceTable(const FShaderCompilerRes
 	BuildResourceTableTokenStream(GenericSRT.ResourceCollectionMap,  GenericSRT.MaxBoundResourceTable, OutSRT.ResourceCollectionMap,  bGenerateEmptyTokenStreamIfNoResources);
 }
 
+static bool DoesUniformBufferNeedReflectedMembers(const TMap<FString, FUniformBufferEntry>& UniformBufferMap, FStringView UniformBufferName)
+{
+	if (const FUniformBufferEntry* Entry = UniformBufferMap.FindByHash(GetTypeHash(UniformBufferName), UniformBufferName))
+	{
+		return EnumHasAnyFlags(Entry->Flags, ERHIUniformBufferFlags::NeedsReflectedMembers);
+	}
+
+	return false;
+}
 
 bool BuildResourceTableMapping(
 	const FShaderResourceTableMap& ResourceTableMap,
@@ -121,16 +130,26 @@ bool BuildResourceTableMapping(
 		// If the shaders uses this member (eg View_PerlinNoise3DTexture)...
 		if (TOptional<FParameterAllocation> Allocation = ParameterMap.FindAndRemoveParameterAllocation(Name))
 		{
+			FStringView UniformBufferName = Entry.GetUniformBufferName();
+
 			const EShaderParameterType ParameterType = Allocation->Type;
 			const bool bBindlessParameter = IsParameterBindless(ParameterType);
 
 			// Force bindless "indices" to zero since they're not needed in SetResourcesFromTables
 			const uint16 BaseIndex = bBindlessParameter ? 0 : Allocation->BaseIndex;
 
+			if (DoesUniformBufferNeedReflectedMembers(UniformBufferMap, UniformBufferName))
+			{
+				FString RenamedMember = Name.Replace(TEXT("_"), TEXT("."));
+				ParameterMap.AddParameterAllocation(*RenamedMember, Allocation->BufferIndex, Allocation->BaseIndex, Allocation->Size, Allocation->Type);
+
+				// Force the parameter to be marked as bound
+				ParameterMap.FindParameterAllocation(RenamedMember);
+			}
+
 			uint16 UniformBufferIndex = INDEX_NONE;
 
 			// Add the UB itself as a parameter if not there
-			FStringView UniformBufferName = Entry.GetUniformBufferName();
 			if (TOptional<FParameterAllocation> UniformBufferParameter = ParameterMap.FindParameterAllocation(UniformBufferName))
 			{
 				UniformBufferIndex = UniformBufferParameter->BufferIndex;
@@ -568,7 +587,7 @@ void HandleReflectedGlobalConstantBufferMember(
 	FShaderCompilerOutput& Output
 )
 {
-	FString MemberName = InMemberName;
+	FStringView MemberName = InMemberName;
 	const EShaderParameterType ParameterType = FShaderParameterParser::ParseAndRemoveBindlessParameterPrefix(MemberName);
 
 	Output.ParameterMap.AddParameterAllocation(
@@ -576,21 +595,30 @@ void HandleReflectedGlobalConstantBufferMember(
 		ConstantBufferIndex,
 		ReflectionOffset,
 		ReflectionSize,
-		ParameterType);
+		ParameterType
+	);
 }
 
 void HandleReflectedUniformBufferConstantBufferMember(
+	EUniformBufferMemberReflectionReason Reason,
+	FStringView UniformBufferName,
 	int32 UniformBufferSlot,
-	const FString& InMemberName,
+	FStringView InMemberName,
 	int32 ReflectionOffset,
 	int32 ReflectionSize,
 	FShaderCompilerOutput& Output
 )
 {
-	FString MemberName = InMemberName;
+	FStringView MemberName = InMemberName;
 	const EShaderParameterType ParameterType = FShaderParameterParser::ParseAndRemoveBindlessParameterPrefix(MemberName);
 
-	if (ParameterType != EShaderParameterType::LooseData)
+	bool bAdd = EnumHasAnyFlags(Reason, EUniformBufferMemberReflectionReason::NeedsReflection);
+	if (EnumHasAnyFlags(Reason, EUniformBufferMemberReflectionReason::Bindless))
+	{
+		bAdd |= (ParameterType != EShaderParameterType::LooseData);
+	}
+
+	if (bAdd)
 	{
 		Output.ParameterMap.AddParameterAllocation(
 			MemberName,
@@ -605,21 +633,27 @@ void HandleReflectedUniformBufferConstantBufferMember(
 void HandleReflectedRootConstantBufferMember(
 	const FShaderCompilerInput& Input,
 	const FShaderParameterParser& ShaderParameterParser,
-	const FString& MemberName,
+	const FString& InMemberName,
 	int32 ReflectionOffset,
 	int32 ReflectionSize,
 	FShaderCompilerOutput& Output
 )
 {
-	ShaderParameterParser.ValidateShaderParameterType(Input, MemberName, ReflectionOffset, ReflectionSize, Output);
+	ShaderParameterParser.ValidateShaderParameterType(Input, InMemberName, ReflectionOffset, ReflectionSize, Output);
 
-	HandleReflectedUniformBufferConstantBufferMember(
-		FShaderParametersMetadata::kRootCBufferBindingIndex,
-		MemberName,
-		ReflectionOffset,
-		ReflectionSize,
-		Output
-	);
+	FStringView MemberName = InMemberName;
+	const EShaderParameterType ParameterType = FShaderParameterParser::ParseAndRemoveBindlessParameterPrefix(MemberName);
+
+	if (ParameterType != EShaderParameterType::LooseData)
+	{
+		Output.ParameterMap.AddParameterAllocation(
+			MemberName,
+			FShaderParametersMetadata::kRootCBufferBindingIndex,
+			ReflectionOffset,
+			1,
+			ParameterType
+		);
+	}
 }
 
 void HandleReflectedRootConstantBuffer(
@@ -654,6 +688,23 @@ void HandleReflectedUniformBuffer(
 	);
 }
 
+EUniformBufferMemberReflectionReason ShouldReflectUniformBufferMembers(const FShaderCompilerInput& Input, FStringView UniformBufferName)
+{
+	EUniformBufferMemberReflectionReason Reason{};
+
+	if (Input.Environment.CompilerFlags.Contains(CFLAG_BindlessResources) || Input.Environment.CompilerFlags.Contains(CFLAG_BindlessSamplers))
+	{
+		Reason |= EUniformBufferMemberReflectionReason::Bindless;
+	}
+
+	if (DoesUniformBufferNeedReflectedMembers(Input.Environment.UniformBufferMap, UniformBufferName))
+	{
+		Reason |= EUniformBufferMemberReflectionReason::NeedsReflection;
+	}
+
+	return Reason;
+}
+
 void HandleReflectedShaderResource(
 	const FString& ResourceName,
 	int32 BindOffset,
diff --git a/Engine/Source/Developer/ShaderCompilerCommon/Public/ShaderCompilerCommon.h b/Engine/Source/Developer/ShaderCompilerCommon/Public/ShaderCompilerCommon.h
--- a/Engine/Source/Developer/ShaderCompilerCommon/Public/ShaderCompilerCommon.h
+++ b/Engine/Source/Developer/ShaderCompilerCommon/Public/ShaderCompilerCommon.h
@@ -191,6 +191,19 @@ namespace UE::ShaderCompilerCommon
 	};
 }
 
+enum class EUniformBufferMemberReflectionReason
+{
+	None            = 0,
+	NeedsReflection = 1 << 0,
+	Bindless        = 1 << 1,
+};
+ENUM_CLASS_FLAGS(EUniformBufferMemberReflectionReason);
+
+SHADERCOMPILERCOMMON_API EUniformBufferMemberReflectionReason ShouldReflectUniformBufferMembers(
+	const FShaderCompilerInput& Input,
+	FStringView UniformBufferName
+);
+
 extern SHADERCOMPILERCOMMON_API void HandleReflectedGlobalConstantBufferMember(
 	const FString& MemberName,
 	uint32 ConstantBufferIndex,
@@ -200,8 +213,10 @@ extern SHADERCOMPILERCOMMON_API void HandleReflectedGlobalConstantBufferMember(
 );
 
 extern SHADERCOMPILERCOMMON_API void HandleReflectedUniformBufferConstantBufferMember(
+	EUniformBufferMemberReflectionReason Reason,
+	FStringView UniformBufferName,
 	int32 UniformBufferSlot,
-	const FString& MemberName,
+	FStringView MemberName,
 	int32 ReflectionOffset,
 	int32 ReflectionSize,
 	FShaderCompilerOutput& CompilerOutput
diff --git a/Engine/Source/Developer/ShaderFormatOpenGL/Private/OpenGLShaderCompiler.cpp b/Engine/Source/Developer/ShaderFormatOpenGL/Private/OpenGLShaderCompiler.cpp
--- a/Engine/Source/Developer/ShaderFormatOpenGL/Private/OpenGLShaderCompiler.cpp
+++ b/Engine/Source/Developer/ShaderFormatOpenGL/Private/OpenGLShaderCompiler.cpp
@@ -1535,7 +1535,7 @@ void ParseReflectionData(const FShaderCompilerInput& ShaderInput, CrossCompiler:
 		{
 			const FUniformBufferEntry* UniformBufferEntry = ShaderInput.Environment.UniformBufferMap.Find(Binding->name);
 
-			if (bEmulatedUBs && (UniformBufferEntry == nullptr || !UniformBufferEntry->bNoEmulatedUniformBuffer))
+			if (bEmulatedUBs && (UniformBufferEntry == nullptr || !EnumHasAnyFlags(UniformBufferEntry->Flags, ERHIUniformBufferFlags::NoEmulatedUniformBuffer)))
 			{
 				check(UBOIndices);
 				uint32 Index = FPlatformMath::CountTrailingZeros(UBOIndices);
diff --git a/Engine/Source/Developer/VulkanShaderFormat/Private/VulkanShaderCompiler.cpp b/Engine/Source/Developer/VulkanShaderFormat/Private/VulkanShaderCompiler.cpp
--- a/Engine/Source/Developer/VulkanShaderFormat/Private/VulkanShaderCompiler.cpp
+++ b/Engine/Source/Developer/VulkanShaderFormat/Private/VulkanShaderCompiler.cpp
@@ -850,24 +850,30 @@ static bool BuildShaderOutputFromSpirv(
 							HandleReflectedUniformBuffer(ResourceName, ReflectionSlot, Output);
 							AddShaderValidationUBSize(BindingTypeCount, Binding->block.padded_size, Output);
 
-							// Register uniform buffer members that are in use
-							for (uint32 MemberIndex = 0; MemberIndex < Binding->block.member_count; ++MemberIndex)
+							const EUniformBufferMemberReflectionReason Reason = ShouldReflectUniformBufferMembers(InternalState.Input, ResourceName);
+							if (Reason != EUniformBufferMemberReflectionReason::None)
 							{
-								const SpvReflectBlockVariable& Member = Binding->block.members[MemberIndex];
-
-								if ((Member.flags & SPV_REFLECT_VARIABLE_FLAGS_UNUSED) != 0)
+								// Register uniform buffer members that are in use
+								for (uint32 MemberIndex = 0; MemberIndex < Binding->block.member_count; ++MemberIndex)
 								{
-									continue;
+									const SpvReflectBlockVariable& Member = Binding->block.members[MemberIndex];
+
+									if ((Member.flags & SPV_REFLECT_VARIABLE_FLAGS_UNUSED) != 0)
+									{
+										continue;
+									}
+
+									const FString MemberName(Member.name);
+									HandleReflectedUniformBufferConstantBufferMember(
+										Reason,
+										ResourceName,
+										ReflectionSlot,
+										MemberName,
+										Member.absolute_offset,
+										Member.size,
+										Output
+									);
 								}
-
-								const FString MemberName(Member.name);
-								HandleReflectedUniformBufferConstantBufferMember(
-									ReflectionSlot,
-									MemberName,
-									Member.absolute_offset,
-									Member.size,
-									Output
-								);
 							}
 						}
 
diff --git a/Engine/Source/Developer/Windows/ShaderFormatD3D/Private/D3DShaderCompiler.inl b/Engine/Source/Developer/Windows/ShaderFormatD3D/Private/D3DShaderCompiler.inl
--- a/Engine/Source/Developer/Windows/ShaderFormatD3D/Private/D3DShaderCompiler.inl
+++ b/Engine/Source/Developer/Windows/ShaderFormatD3D/Private/D3DShaderCompiler.inl
@@ -107,9 +107,12 @@ template <typename ID3D1xShaderReflection, typename D3D1x_SHADER_DESC, typename
 			ID3D1xShaderReflectionConstantBuffer* ConstantBuffer = Reflector->GetConstantBufferByName(BindDesc.Name);
 			D3D1x_SHADER_BUFFER_DESC CBDesc;
 			ConstantBuffer->GetDesc(&CBDesc);
-			const bool bGlobalCB = (FCStringAnsi::Strcmp(CBDesc.Name, "$Globals") == 0);
-			const bool bRootConstantsCB = (FCStringAnsi::Strcmp(CBDesc.Name, "UERootConstants") == 0);
-			const bool bIsRootCB = FCString::Strcmp(ANSI_TO_TCHAR(CBDesc.Name), FShaderParametersMetadata::kRootUniformBufferBindingName) == 0;
+
+			const FString ConstantBufferName(CBDesc.Name);
+
+			const bool bGlobalCB = (ConstantBufferName == TEXT("$Globals"));
+			const bool bRootConstantsCB = (ConstantBufferName == TEXT("UERootConstants"));
+			const bool bIsRootCB = (ConstantBufferName == FShaderParametersMetadata::kRootUniformBufferBindingName);
 
 			if (bGlobalCB)
 			{
@@ -211,14 +214,13 @@ template <typename ID3D1xShaderReflection, typename D3D1x_SHADER_DESC, typename
 			else
 			{
 				// Track just the constant buffer itself.
-				const FString UniformBufferName(CBDesc.Name);
-
 				AddShaderValidationUBSize(CBIndex, CBDesc.Size, Output);
-				HandleReflectedUniformBuffer(UniformBufferName, CBIndex, Output);
+				HandleReflectedUniformBuffer(ConstantBufferName, CBIndex, Output);
 				
 				CompileData.UsedUniformBufferSlots[CBIndex] = true;
 
-				if (bBindlessEnabled)
+				const EUniformBufferMemberReflectionReason Reason = ShouldReflectUniformBufferMembers(Input, ConstantBufferName);
+				if (Reason != EUniformBufferMemberReflectionReason::None)
 				{
 					for (uint32 ConstantIndex = 0; ConstantIndex < CBDesc.Variables; ConstantIndex++)
 					{
@@ -229,9 +231,13 @@ template <typename ID3D1xShaderReflection, typename D3D1x_SHADER_DESC, typename
 
 						if (VariableDesc.uFlags & D3D_SVF_USED)
 						{
+							const FString MemberName(VariableDesc.Name);
+
 							HandleReflectedUniformBufferConstantBufferMember(
+								Reason,
+								ConstantBufferName,
 								CBIndex,
-								FString(VariableDesc.Name),
+								MemberName,
 								VariableDesc.StartOffset,
 								VariableDesc.Size,
 								Output
@@ -245,7 +251,7 @@ template <typename ID3D1xShaderReflection, typename D3D1x_SHADER_DESC, typename
 			{
 				CompileData.UniformBufferNames.AddDefaulted(CBIndex - CompileData.UniformBufferNames.Num() + 1);
 			}
-			CompileData.UniformBufferNames[CBIndex] = UE::ShaderCompilerCommon::RemoveConstantBufferPrefix(FString(CBDesc.Name));
+			CompileData.UniformBufferNames[CBIndex] = UE::ShaderCompilerCommon::RemoveConstantBufferPrefix(ConstantBufferName);
 
 			CompileData.NumCBs = FMath::Max(CompileData.NumCBs, BindDesc.BindPoint + BindDesc.BindCount);
 		}
@@ -344,7 +350,7 @@ template <typename ID3D1xShaderReflection, typename D3D1x_SHADER_DESC, typename
 
 			CompileData.NumSRVs = FMath::Max(CompileData.NumSRVs, BindDesc.BindPoint + 1);
 		}
-		else if (BindDesc.Type == (D3D_SHADER_INPUT_TYPE)(D3D_SIT_UAV_RWSTRUCTURED_WITH_COUNTER + 1)) // D3D_SIT_RTACCELERATIONSTRUCTURE (12)
+		else if (BindDesc.Type == D3D_SIT_RTACCELERATIONSTRUCTURE)
 		{
 			// Acceleration structure resources are treated as SRVs.
 			check(BindDesc.BindCount == 1);
diff --git a/Engine/Source/Runtime/OpenGLDrv/Private/OpenGLUniformBuffer.cpp b/Engine/Source/Runtime/OpenGLDrv/Private/OpenGLUniformBuffer.cpp
--- a/Engine/Source/Runtime/OpenGLDrv/Private/OpenGLUniformBuffer.cpp
+++ b/Engine/Source/Runtime/OpenGLDrv/Private/OpenGLUniformBuffer.cpp
@@ -419,7 +419,7 @@ FOpenGLUniformBuffer::FOpenGLUniformBuffer(const FRHIUniformBufferLayout* InLayo
 	, bStreamDraw(false)
 	, bOwnsResource(true)
 {
-	bIsEmulatedUniformBuffer = GUseEmulatedUniformBuffers && !(InLayout->bNoEmulatedUniformBuffer || InLayout->bUniformView);
+	bIsEmulatedUniformBuffer = GUseEmulatedUniformBuffers && !EnumHasAnyFlags(InLayout->Flags, ERHIUniformBufferFlags::NoEmulatedUniformBuffer|ERHIUniformBufferFlags::UniformView);
 	RangeSize = InLayout->ConstantBufferSize;
 }
 
@@ -462,7 +462,7 @@ FOpenGLUniformBuffer::~FOpenGLUniformBuffer()
 
 			FScopeLock Lock(&GGLUniformBufferPoolCS);
 			
-			if (GUseEmulatedUniformBuffers && !GetLayout().bNoEmulatedUniformBuffer)
+			if (GUseEmulatedUniformBuffers && !EnumHasAnyFlags(GetLayout().Flags, ERHIUniformBufferFlags::NoEmulatedUniformBuffer))
 			{
 				SafeGLEmulatedUniformBufferPools[SafeFrameIndex][BucketIndex][StreamedIndex].Add(NewEntry);
 			}
@@ -556,7 +556,7 @@ static FUniformBufferRHIRef CreateUniformBuffer(const void* Contents, const FRHI
 	// PersistentlyMappedBuffer initializes via IsSuballocatingUBOs path which will flush RHI commands. safe to use on RT thread.
 	uint8* PersistentlyMappedBuffer = NULL;
 
-	bool bUseEmulatedUBs = GUseEmulatedUniformBuffers && !Layout->bNoEmulatedUniformBuffer;
+	bool bUseEmulatedUBs = GUseEmulatedUniformBuffers && !EnumHasAnyFlags(Layout->Flags, ERHIUniformBufferFlags::NoEmulatedUniformBuffer);
 
 	FRHICommandListImmediate& RHICmdList = FRHICommandListImmediate::Get();
 
@@ -628,7 +628,7 @@ static FOpenGLUniformBuffer* CreateUniformBufferView(FRHICommandListImmediate& R
 {
 	FOpenGLUniformBuffer* UniformBufferView = nullptr;
 	
-	if (Layout->bUniformView)
+	if (EnumHasAnyFlags(Layout->Flags, ERHIUniformBufferFlags::UniformView))
 	{
 		UniformBufferView = new FOpenGLUniformBuffer(Layout);
 		UniformBufferView->SetLayoutTable(Contents, EUniformBufferValidation::None);
@@ -692,7 +692,7 @@ FUniformBufferRHIRef FOpenGLDynamicRHI::RHICreateUniformBuffer(const void* Conte
 		return UniformBufferView;
 	}
 
-	bool bUseEmulatedUBs = GUseEmulatedUniformBuffers && !Layout->bNoEmulatedUniformBuffer;
+	bool bUseEmulatedUBs = GUseEmulatedUniformBuffers && !EnumHasAnyFlags(Layout->Flags, ERHIUniformBufferFlags::NoEmulatedUniformBuffer);
 
 	bool bStreamDraw = (Usage == UniformBuffer_SingleDraw || Usage == UniformBuffer_SingleFrame);
 	GLuint AllocatedResource = 0;
diff --git a/Engine/Source/Runtime/RHI/Private/RHIResources.cpp b/Engine/Source/Runtime/RHI/Private/RHIResources.cpp
--- a/Engine/Source/Runtime/RHI/Private/RHIResources.cpp
+++ b/Engine/Source/Runtime/RHI/Private/RHIResources.cpp
@@ -440,9 +440,7 @@ FRHIUniformBufferLayout::FRHIUniformBufferLayout(const FRHIUniformBufferLayoutIn
 	, RenderTargetsOffset(Initializer.RenderTargetsOffset)
 	, StaticSlot(Initializer.StaticSlot)
 	, BindingFlags(Initializer.BindingFlags)
-	, bHasNonGraphOutputs(Initializer.bHasNonGraphOutputs)
-	, bNoEmulatedUniformBuffer(Initializer.bNoEmulatedUniformBuffer)
-	, bUniformView(Initializer.bUniformView)
+	, Flags(Initializer.Flags)
 {
 }
 
diff --git a/Engine/Source/Runtime/RHI/Public/RHIDefinitions.h b/Engine/Source/Runtime/RHI/Public/RHIDefinitions.h
--- a/Engine/Source/Runtime/RHI/Public/RHIDefinitions.h
+++ b/Engine/Source/Runtime/RHI/Public/RHIDefinitions.h
@@ -618,6 +618,25 @@ enum class EUniformBufferBindingFlags : uint8
 };
 ENUM_CLASS_FLAGS(EUniformBufferBindingFlags);
 
+/** Flags for Uniform Buffers */
+enum class ERHIUniformBufferFlags : uint8
+{
+	None                    = 0,
+
+	/** Whether to force a real uniform buffer when using emulated uniform buffers */
+	NoEmulatedUniformBuffer = 1 << 0,
+
+	/** Signals if the uniform buffer members need to be included in shader reflection */
+	NeedsReflectedMembers   = 1 << 1,
+
+	/** Whether this layout may contain non-render-graph outputs (e.g. RHI UAVs). */
+	HasNonGraphOutputs      = 1 << 2,
+
+	/** This struct is a view into uniform buffer object, on platforms that support UBO */
+	UniformView             = 1 << 3,
+};
+ENUM_CLASS_FLAGS(ERHIUniformBufferFlags);
+
 /** Numerical type used to store the static slot indices. */
 using FUniformBufferStaticSlot = uint8;
 
diff --git a/Engine/Source/Runtime/RHI/Public/RHIMemoryLayout.h b/Engine/Source/Runtime/RHI/Public/RHIMemoryLayout.h
--- a/Engine/Source/Runtime/RHI/Public/RHIMemoryLayout.h
+++ b/Engine/Source/Runtime/RHI/Public/RHIMemoryLayout.h
@@ -9,4 +9,5 @@
 DECLARE_INTRINSIC_TYPE_LAYOUT(EVertexElementType);
 DECLARE_INTRINSIC_TYPE_LAYOUT(EUniformBufferBaseType);
 DECLARE_INTRINSIC_TYPE_LAYOUT(EUniformBufferBindingFlags);
+DECLARE_INTRINSIC_TYPE_LAYOUT(ERHIUniformBufferFlags);
 DECLARE_INTRINSIC_TYPE_LAYOUT(ERHIFeatureLevel::Type);
diff --git a/Engine/Source/Runtime/RHI/Public/RHIResources.h b/Engine/Source/Runtime/RHI/Public/RHIResources.h
--- a/Engine/Source/Runtime/RHI/Public/RHIResources.h
+++ b/Engine/Source/Runtime/RHI/Public/RHIResources.h
@@ -1174,7 +1174,7 @@ struct FRHIUniformBufferLayout : public FRHIResource
 
 	inline bool HasExternalOutputs() const
 	{
-		return bHasNonGraphOutputs;
+		return EnumHasAnyFlags(Flags, ERHIUniformBufferFlags::HasNonGraphOutputs);
 	}
 
 	inline bool HasStaticSlot() const
@@ -1216,14 +1216,8 @@ struct FRHIUniformBufferLayout : public FRHIResource
 	/** The binding flags describing how this resource can be bound to the RHI. */
 	const EUniformBufferBindingFlags BindingFlags;
 
-	/** Whether this layout may contain non-render-graph outputs (e.g. RHI UAVs). */
-	const bool bHasNonGraphOutputs;
-
-	/** Used for platforms which use emulated ub's, forces a real uniform buffer instead */
-	const bool bNoEmulatedUniformBuffer;
-
-	/** This struct is a view into uniform buffer object, on platforms that support UBO */
-	const bool bUniformView;
+	/** Flags to signal different Uniform Buffer states. */
+	const ERHIUniformBufferFlags Flags;
 
 	/** Compare two uniform buffer layouts. */
 	friend inline bool operator==(const FRHIUniformBufferLayout& A, const FRHIUniformBufferLayout& B)
diff --git a/Engine/Source/Runtime/RHI/Public/RHIUniformBufferLayoutInitializer.h b/Engine/Source/Runtime/RHI/Public/RHIUniformBufferLayoutInitializer.h
--- a/Engine/Source/Runtime/RHI/Public/RHIUniformBufferLayoutInitializer.h
+++ b/Engine/Source/Runtime/RHI/Public/RHIUniformBufferLayoutInitializer.h
@@ -61,8 +61,18 @@ struct FRHIUniformBufferLayoutInitializer
 
 	void ComputeHash()
 	{
+		uint32 TmpHash = 0;
+
+		TmpHash |= (ConstantBufferSize << 16);
+		TmpHash |= (static_cast<uint32>(BindingFlags) << 8);
+
 		// Static slot is not stable. Just track whether we have one at all.
-		uint32 TmpHash = ConstantBufferSize << 16 | static_cast<uint32>(BindingFlags) << 8 | static_cast<uint32>(StaticSlot != MAX_UNIFORM_BUFFER_STATIC_SLOTS);
+		TmpHash |= (StaticSlot != MAX_UNIFORM_BUFFER_STATIC_SLOTS ? 1 << 0 : 0);
+
+		// These flags aren't based on the resources, so make sure they change the hash
+		TmpHash |= (EnumHasAnyFlags(Flags, ERHIUniformBufferFlags::NoEmulatedUniformBuffer) ? 1 << 1 : 0);
+		TmpHash |= (EnumHasAnyFlags(Flags, ERHIUniformBufferFlags::NeedsReflectedMembers)   ? 1 << 2 : 0);
+		TmpHash |= (EnumHasAnyFlags(Flags, ERHIUniformBufferFlags::UniformView)             ? 1 << 3 : 0);
 
 		for (int32 ResourceIndex = 0; ResourceIndex < Resources.Num(); ResourceIndex++)
 		{
@@ -96,6 +106,7 @@ struct FRHIUniformBufferLayoutInitializer
 		ConstantBufferSize = Source.ConstantBufferSize;
 		StaticSlot = Source.StaticSlot;
 		BindingFlags = Source.BindingFlags;
+		Flags = Source.Flags;
 		Resources = Source.Resources;
 		Name = Source.Name;
 		Hash = Source.Hash;
@@ -111,11 +122,6 @@ struct FRHIUniformBufferLayoutInitializer
 		return RenderTargetsOffset != kUniformBufferInvalidOffset;
 	}
 
-	bool HasExternalOutputs() const
-	{
-		return bHasNonGraphOutputs;
-	}
-
 	bool HasStaticSlot() const
 	{
 		return IsUniformBufferStaticSlotValid(StaticSlot);
@@ -126,8 +132,8 @@ struct FRHIUniformBufferLayoutInitializer
 		Ar << Ref.ConstantBufferSize;
 		Ar << Ref.StaticSlot;
 		Ar << Ref.RenderTargetsOffset;
-		Ar << Ref.bHasNonGraphOutputs;
 		Ar << Ref.BindingFlags;
+		Ar << Ref.Flags;
 		Ar << Ref.Resources;
 		Ar << Ref.GraphResources;
 		Ar << Ref.GraphTextures;
@@ -136,8 +142,6 @@ struct FRHIUniformBufferLayoutInitializer
 		Ar << Ref.UniformBuffers;
 		Ar << Ref.Name;
 		Ar << Ref.Hash;
-		Ar << Ref.bNoEmulatedUniformBuffer;
-		Ar << Ref.bUniformView;
 		return Ar;
 	}
 
@@ -180,14 +184,8 @@ public:
 	/** The binding flags describing how this resource can be bound to the RHI. */
 	LAYOUT_FIELD_INITIALIZED(EUniformBufferBindingFlags, BindingFlags, EUniformBufferBindingFlags::Shader);
 
-	/** Whether this layout may contain non-render-graph outputs (e.g. RHI UAVs). */
-	LAYOUT_FIELD_INITIALIZED(bool, bHasNonGraphOutputs, false);
-
-	/** Used for platforms which use emulated ub's, forces a real uniform buffer instead */
-	LAYOUT_FIELD_INITIALIZED(bool, bNoEmulatedUniformBuffer, false);
-
-	/** This struct is a view into uniform buffer object, on platforms that support UBO */
-	LAYOUT_FIELD_INITIALIZED(bool, bUniformView, false);
+	/** Flags for the uniform buffer */
+	LAYOUT_FIELD_INITIALIZED(ERHIUniformBufferFlags, Flags, ERHIUniformBufferFlags::None);
 	
 	/** Compare two uniform buffer layout initializers. */
 	friend inline bool operator==(const FRHIUniformBufferLayoutInitializer& A, const FRHIUniformBufferLayoutInitializer& B)
diff --git a/Engine/Source/Runtime/RenderCore/Private/ShaderParameterMetadata.cpp b/Engine/Source/Runtime/RenderCore/Private/ShaderParameterMetadata.cpp
--- a/Engine/Source/Runtime/RenderCore/Private/ShaderParameterMetadata.cpp
+++ b/Engine/Source/Runtime/RenderCore/Private/ShaderParameterMetadata.cpp
@@ -412,7 +412,7 @@ FShaderParametersMetadata::FShaderParametersMetadata(
 	const TArray<FMember>& InMembers,
 	bool bForceCompleteInitialization,
 	FRHIUniformBufferLayoutInitializer* OutLayoutInitializer,
-	uint32 InUsageFlags)
+	EUsageFlags InUsageFlags)
 	: LayoutName(InLayoutName)
 	, StructTypeName(InStructTypeName)
 	, ShaderVariableName(InShaderVariableName)
@@ -423,9 +423,9 @@ FShaderParametersMetadata::FShaderParametersMetadata(
 	, Size(InSize)
 	, UseCase(InUseCase)
 	, BindingFlags(InBindingFlags)
+	, UsageFlags(InUsageFlags)
 	, Members(InMembers)
 	, GlobalListLink(this)
-	, UsageFlags(InUsageFlags)
 {
 	checkf(UseCase == EUseCase::UniformBuffer || !EnumHasAnyFlags(BindingFlags, EUniformBufferBindingFlags::Static), TEXT("Only uniform buffers can utilize the global binding flag."));
 
@@ -565,8 +565,21 @@ void FShaderParametersMetadata::InitializeLayout(FRHIUniformBufferLayoutInitiali
 	FRHIUniformBufferLayoutInitializer LocalLayoutInitializer(LayoutName);
 	FRHIUniformBufferLayoutInitializer& LayoutInitializer = OutLayoutInitializer ? *OutLayoutInitializer : LocalLayoutInitializer;
 	LayoutInitializer.ConstantBufferSize = Size;
-	LayoutInitializer.bUniformView = UsageFlags & (uint32)EUsageFlags::UniformView;
-	LayoutInitializer.bNoEmulatedUniformBuffer = LayoutInitializer.bUniformView || (UsageFlags & (uint32)EUsageFlags::NoEmulatedUniformBuffer);
+
+	if (EnumHasAnyFlags(UsageFlags, EUsageFlags::UniformView))
+	{
+		EnumAddFlags(LayoutInitializer.Flags, ERHIUniformBufferFlags::UniformView);
+	}
+
+	if (EnumHasAnyFlags(UsageFlags, EUsageFlags::NoEmulatedUniformBuffer|EUsageFlags::UniformView))
+	{
+		EnumAddFlags(LayoutInitializer.Flags, ERHIUniformBufferFlags::NoEmulatedUniformBuffer);
+	}
+
+	if (EnumHasAnyFlags(UsageFlags, EUsageFlags::NeedsReflectedMembers))
+	{
+		EnumAddFlags(LayoutInitializer.Flags, ERHIUniformBufferFlags::NeedsReflectedMembers);
+	}
 	
 	if (StaticSlotName)
 	{
@@ -610,6 +623,8 @@ void FShaderParametersMetadata::InitializeLayout(FRHIUniformBufferLayoutInitiali
 	/** Allow all use cases that inline a structure within another. Data driven are not known to inline structures. */
 	const bool bAllowStructureInlining = UseCase == EUseCase::ShaderParameterStruct || UseCase == EUseCase::UniformBuffer;
 
+	bool bHasNonGraphOutputs = false;
+
 	for (int32 i = 0; i < MemberStack.Num(); ++i)
 	{
 		const FShaderParametersMetadata& CurrentStruct = MemberStack[i].ContainingStruct;
@@ -625,7 +640,7 @@ void FShaderParametersMetadata::InitializeLayout(FRHIUniformBufferLayoutInitiali
 		const bool bIsRDGResource = IsRDGResourceReferenceShaderParameterType(BaseType);
 		const bool bIsVariableNativeType = CurrentMember.IsVariableNativeType();
 
-		LayoutInitializer.bHasNonGraphOutputs |= BaseType == UBMT_UAV;
+		bHasNonGraphOutputs |= (BaseType == UBMT_UAV);
 
 		if (DO_CHECK)
 		{
@@ -772,6 +787,11 @@ void FShaderParametersMetadata::InitializeLayout(FRHIUniformBufferLayoutInitiali
 			}
 		}
 
+		if (bHasNonGraphOutputs)
+		{
+			EnumAddFlags(LayoutInitializer.Flags, ERHIUniformBufferFlags::HasNonGraphOutputs);
+		}
+
 		if (IsShaderParameterTypeForUniformBufferLayout(BaseType))
 		{
 			for (uint32 ArrayElementId = 0; ArrayElementId < (bIsArray ? ArraySize : 1u); ArrayElementId++)
@@ -1027,20 +1047,21 @@ void FShaderParametersMetadata::AddResourceTableEntries(FShaderResourceTableMap&
 	
 	FUniformBufferEntry UniformBufferEntry;
 	UniformBufferEntry.StaticSlotName = StaticSlotName;
+	UniformBufferEntry.MemberNameBuffer = MemberNameBuffer;
 	UniformBufferEntry.LayoutHash = IsLayoutInitialized() ? GetLayout().GetHash() : 0;
 	UniformBufferEntry.BindingFlags = BindingFlags;
-	UniformBufferEntry.bNoEmulatedUniformBuffer = (UsageFlags & (uint32)EUsageFlags::NoEmulatedUniformBuffer) || (UsageFlags & (uint32)EUsageFlags::UniformView);
-	UniformBufferEntry.MemberNameBuffer = MemberNameBuffer;
-	UniformBufferMap.AddByHash(ShaderVariableNameHash, ShaderVariableName, UniformBufferEntry);
-}
 
-PRAGMA_DISABLE_DEPRECATION_WARNINGS
-// Deprecated version of function
-void FShaderParametersMetadata::AddResourceTableEntries(TMap<FString, FResourceTableEntry>& ResourceTableMap, TMap<FString, FUniformBufferEntry>& UniformBufferMap) const
-{
-	UE_LOG(LogShaders, Error, TEXT("FShaderParametersMetadata::AddResourceTableEntries call that accepts a TMap has been deprecated.  Use FShaderResourceTableMap structure instead."));
+	if (EnumHasAnyFlags(UsageFlags, EUsageFlags::NoEmulatedUniformBuffer | EUsageFlags::UniformView))
+	{
+		EnumAddFlags(UniformBufferEntry.Flags, ERHIUniformBufferFlags::NoEmulatedUniformBuffer);
+	}
+	if (EnumHasAnyFlags(UsageFlags, EUsageFlags::NeedsReflectedMembers))
+	{
+		EnumAddFlags(UniformBufferEntry.Flags, ERHIUniformBufferFlags::NeedsReflectedMembers);
+	}
+
+	UniformBufferMap.AddByHash(ShaderVariableNameHash, ShaderVariableName, UniformBufferEntry);
 }
-PRAGMA_ENABLE_DEPRECATION_WARNINGS
 
 #endif // WITH_EDITOR
 
diff --git a/Engine/Source/Runtime/RenderCore/Private/ShaderParameterMetadataBuilder.cpp b/Engine/Source/Runtime/RenderCore/Private/ShaderParameterMetadataBuilder.cpp
--- a/Engine/Source/Runtime/RenderCore/Private/ShaderParameterMetadataBuilder.cpp
+++ b/Engine/Source/Runtime/RenderCore/Private/ShaderParameterMetadataBuilder.cpp
@@ -205,7 +205,7 @@ FShaderParametersMetadata* FShaderParametersMetadataBuilder::Build(
 	const int32 InFileLine,
 	bool bForceCompleteInitialization,
 	FRHIUniformBufferLayoutInitializer* OutLayoutInitializer,
-	uint32 InUsageFlags
+	FShaderParametersMetadata::EUsageFlags InUsageFlags
 )
 {
 	const uint32 StructSize = Align(NextMemberOffset, SHADER_PARAMETER_STRUCT_ALIGNMENT);
diff --git a/Engine/Source/Runtime/RenderCore/Private/ShaderParameters.cpp b/Engine/Source/Runtime/RenderCore/Private/ShaderParameters.cpp
--- a/Engine/Source/Runtime/RenderCore/Private/ShaderParameters.cpp
+++ b/Engine/Source/Runtime/RenderCore/Private/ShaderParameters.cpp
@@ -17,6 +17,7 @@
 IMPLEMENT_TYPE_LAYOUT(FShaderParameter);
 IMPLEMENT_TYPE_LAYOUT(FShaderResourceParameter);
 IMPLEMENT_TYPE_LAYOUT(FShaderUniformBufferParameter);
+IMPLEMENT_TYPE_LAYOUT(FShaderUniformBufferMemberParameter);
 
 static void FailureToBindNonOptionalParameter(const TCHAR* ParameterType, const TCHAR* ParameterName)
 {
@@ -72,6 +73,16 @@ FArchive& operator<<(FArchive& Ar,FShaderResourceParameter& P)
 	return Ar << P.BaseIndex << P.NumResources;
 }
 
+void FShaderUniformBufferMemberParameter::Bind(const FShaderParameterMap& ParameterMap, const TCHAR* ParameterName)
+{
+	bIsBound = ParameterMap.ContainsParameterAllocation(ParameterName) ? 1 : 0;
+}
+
+FArchive& operator<<(FArchive& Ar, FShaderUniformBufferMemberParameter& P)
+{
+	return Ar << P.bIsBound;
+}
+
 #if WITH_EDITOR
 void FShaderUniformBufferParameter::ModifyCompilationEnvironment(const TCHAR* ParameterName, const FShaderParametersMetadata& Struct, EShaderPlatform Platform, FShaderCompilerEnvironment& OutEnvironment)
 {
@@ -160,7 +171,7 @@ static void CreateHLSLUniformBufferStructMembersDeclaration(
 		Decl.StructMembers << Member.GetName() << TEXT(");\n");
 	};
 
-	if ((UniformBufferStruct.GetUsageFlags() & (uint32)FShaderParametersMetadata::EUsageFlags::UniformView) != 0)
+	if (EnumHasAnyFlags(UniformBufferStruct.GetUsageFlags(), FShaderParametersMetadata::EUsageFlags::UniformView))
 	{
 		// UniformView struct is expected to have a single SRV member which serves as a uniform view
 		check(StructMembers.Num() == 1);
diff --git a/Engine/Source/Runtime/RenderCore/Public/ShaderCore.h b/Engine/Source/Runtime/RenderCore/Public/ShaderCore.h
--- a/Engine/Source/Runtime/RenderCore/Public/ShaderCore.h
+++ b/Engine/Source/Runtime/RenderCore/Public/ShaderCore.h
@@ -366,14 +366,14 @@ inline FArchive& operator<<(FArchive& Ar, FUniformResourceEntry& Entry)
 inline FArchive& operator<<(FArchive& Ar, FUniformBufferEntry& Entry)
 {
 	Ar << Entry.StaticSlotName;
-	Ar << Entry.LayoutHash;
-	Ar << Entry.BindingFlags;
-	Ar << Entry.bNoEmulatedUniformBuffer;
 	if (Ar.IsLoading())
 	{
 		Entry.MemberNameBuffer = MakeShareable(new TArray<TCHAR>());
 	}
 	Ar << *Entry.MemberNameBuffer.Get();
+	Ar << Entry.LayoutHash;
+	Ar << Entry.BindingFlags;
+	Ar << Entry.Flags;
 	return Ar;
 }
 
diff --git a/Engine/Source/Runtime/RenderCore/Public/ShaderParameterMacros.h b/Engine/Source/Runtime/RenderCore/Public/ShaderParameterMacros.h
--- a/Engine/Source/Runtime/RenderCore/Public/ShaderParameterMacros.h
+++ b/Engine/Source/Runtime/RenderCore/Public/ShaderParameterMacros.h
@@ -1535,7 +1535,7 @@ private:
 	TFunctionRef<const FShaderParametersMetadata* ()> LazyShaderParametersMetadataAccessor;
 };
 
-#define IMPLEMENT_UNIFORM_BUFFER_STRUCT(StructTypeName,ShaderVariableName) \
+#define IMPLEMENT_UNIFORM_BUFFER_STRUCT_EX(StructTypeName,ShaderVariableName,UsageFlags) \
 	const FShaderParametersMetadata* GetForwardDeclaredShaderParametersStructMetadata(const StructTypeName* DummyPtr) { return StructTypeName::FTypeInfo::GetStructMetadata(); } \
 	const FShaderParametersMetadata* StructTypeName::GetStructMetadata() \
 	{ \
@@ -1550,11 +1550,17 @@ private:
 			StructTypeName::FTypeInfo::FileName, \
 			StructTypeName::FTypeInfo::FileLine, \
 			sizeof(StructTypeName), \
-			StructTypeName::zzGetMembers()); \
+			StructTypeName::zzGetMembers(), \
+			false, \
+			nullptr, \
+			UsageFlags); \
 		return &StaticStructMetadata; \
 	} \
 	FShaderParametersMetadataRegistration StructTypeName##MetadataRegistration { TFunctionRef<const ::FShaderParametersMetadata* ()>{StructTypeName::GetStructMetadata} };
 
+#define IMPLEMENT_UNIFORM_BUFFER_STRUCT(StructTypeName,ShaderVariableName) \
+	IMPLEMENT_UNIFORM_BUFFER_STRUCT_EX(StructTypeName,ShaderVariableName,FShaderParametersMetadata::EUsageFlags::None)
+
 #define IMPLEMENT_UNIFORM_BUFFER_ALIAS_STRUCT(StructTypeName, UniformBufferAlias) \
 	static const FShaderParametersMetadata UniformBufferAlias( \
 	FShaderParametersMetadata::EUseCase::UniformBuffer, \
@@ -1608,7 +1614,7 @@ private:
 			StructTypeName::zzGetMembers(), \
 			false, \
 			nullptr, \
-			(uint32)UsageFlags); \
+			UsageFlags); \
 		return &StaticStructMetadata; \
 	} \
 	FShaderParametersMetadataRegistration StructTypeName##MetadataRegistration { TFunctionRef<const ::FShaderParametersMetadata* ()>{StructTypeName::GetStructMetadata} };
diff --git a/Engine/Source/Runtime/RenderCore/Public/ShaderParameterMetadata.h b/Engine/Source/Runtime/RenderCore/Public/ShaderParameterMetadata.h
--- a/Engine/Source/Runtime/RenderCore/Public/ShaderParameterMetadata.h
+++ b/Engine/Source/Runtime/RenderCore/Public/ShaderParameterMetadata.h
@@ -61,26 +61,23 @@ struct FUniformResourceEntry
 	FORCEINLINE FStringView GetUniformBufferName() const { return FStringView(UniformBufferMemberName, UniformBufferNameLength); }
 };
 
-struct UE_DEPRECATED(5.3, "Deprecated structure -- replaced with FUniformResourceEntry.") FResourceTableEntry
-{
-	FString UniformBufferName;
-	uint16 Type{};
-	uint16 ResourceIndex{};
-};
-
 /** Minimal information about each uniform buffer entry fed to the shader compiler. */
 struct FUniformBufferEntry
 {
 	/** The name of the uniform buffer static slot (if global). */
 	FString StaticSlotName;
+
+	/** Storage for member names for this uniform buffer (pointed to by FUniformResourceEntry::UniformBufferMemberName)  */
+	FThreadSafeNameBufferPtr MemberNameBuffer;
+
 	/** Hash of the resource table layout. */
 	uint32 LayoutHash{};
+
 	/** The binding flags used by this resource table. */
 	EUniformBufferBindingFlags BindingFlags{ EUniformBufferBindingFlags::Shader };
+
 	/** Whether to force a real uniform buffer when using emulated uniform buffers */
-	bool bNoEmulatedUniformBuffer;
-	/** Storage for member names for this uniform buffer (pointed to by FUniformResourceEntry::UniformBufferMemberName)  */
-	FThreadSafeNameBufferPtr MemberNameBuffer;
+	ERHIUniformBufferFlags Flags{ ERHIUniformBufferFlags::None };
 };
 
 /** Parse the shader resource binding from the binding type used in shader code. */
@@ -161,6 +158,12 @@ public:
 		
 		/** This struct is a view into uniform buffer object, on platforms that support UBO */
 		UniformView = 1 << 1,
+
+		/** This struct needs its members reflected for binding information. */
+		NeedsReflectedMembers = 1 << 2,
+
+		/** Signals that the uniform buffer is manually bound by the pass and should be ignored by the mesh pass processor. */
+		ManuallyBoundByPass = 1 << 3,
 	};
 
 	/** Shader binding name of the uniform buffer that contains the root shader parameters. */
@@ -329,7 +332,7 @@ public:
 		const TArray<FMember>& InMembers,
 		bool bForceCompleteInitialization = false,
 		FRHIUniformBufferLayoutInitializer* OutLayoutInitializer = nullptr,
-		uint32 InUsageFlags = 0);
+		EUsageFlags InUsageFlags = EUsageFlags::None);
 
 	RENDERCORE_API virtual ~FShaderParametersMetadata();
 
@@ -337,11 +340,6 @@ public:
 
 #if WITH_EDITOR
 	RENDERCORE_API void AddResourceTableEntries(FShaderResourceTableMap& ResourceTableMap, TMap<FString, FUniformBufferEntry>& UniformBufferMap) const;
-
-	PRAGMA_DISABLE_DEPRECATION_WARNINGS
-	UE_DEPRECATED(5.3, "Resource table entries are now stored in FShaderResourceTableMap, rather than a TMap.")
-	RENDERCORE_API void AddResourceTableEntries(TMap<FString, FResourceTableEntry>& ResourceTableMap, TMap<FString, FUniformBufferEntry>& UniformBufferMap) const;
-	PRAGMA_ENABLE_DEPRECATION_WARNINGS
 #endif
 
 	const TCHAR* GetStructTypeName() const { return StructTypeName; }
@@ -370,7 +368,7 @@ public:
 	uint32 GetSize() const { return Size; }
 	EUseCase GetUseCase() const { return UseCase; }
 	inline bool IsLayoutInitialized() const { return Layout != nullptr; }
-	uint32 GetUsageFlags() const { return UsageFlags; }
+	EUsageFlags GetUsageFlags() const { return UsageFlags; }
 
 	const FRHIUniformBufferLayout& GetLayout() const
 	{
@@ -488,6 +486,9 @@ private:
 	/** The binding model used by this parameter struct. */
 	const EUniformBufferBindingFlags BindingFlags;
 
+	/** Additional flags for how to use the buffer */
+	const EUsageFlags UsageFlags;
+
 	/** Layout of all the resources in the shader parameter struct. */
 	FUniformBufferLayoutRHIRef Layout{};
 	
@@ -533,9 +534,6 @@ private:
 	FBlake3Hash LayoutSignature;
 #endif
 
-	/** Additional flags for how to use the buffer */
-	uint32 UsageFlags = 0;
-
 	RENDERCORE_API void InitializeLayout(FRHIUniformBufferLayoutInitializer* OutLayoutInitializer = nullptr);
 
 #if WITH_EDITOR
@@ -543,6 +541,8 @@ private:
 #endif
 };
 
+ENUM_CLASS_FLAGS(FShaderParametersMetadata::EUsageFlags);
+
 
 /**
  * Utility class for caching FName and other info for a shader compiler define.  Don't use directly, use SET_SHADER_DEFINE or SET_SHADER_DEFINE_AND_COMPILE_ARGUMENT
diff --git a/Engine/Source/Runtime/RenderCore/Public/ShaderParameterMetadataBuilder.h b/Engine/Source/Runtime/RenderCore/Public/ShaderParameterMetadataBuilder.h
--- a/Engine/Source/Runtime/RenderCore/Public/ShaderParameterMetadataBuilder.h
+++ b/Engine/Source/Runtime/RenderCore/Public/ShaderParameterMetadataBuilder.h
@@ -187,7 +187,7 @@ public:
 		const int32 InFileLine,
 		bool bForceCompleteInitialization = false,
 		FRHIUniformBufferLayoutInitializer* OutLayoutInitializer = nullptr,
-		uint32 InUsageFlags = 0
+		FShaderParametersMetadata::EUsageFlags InUsageFlags = FShaderParametersMetadata::EUsageFlags::None
 		);
 
 private:
diff --git a/Engine/Source/Runtime/RenderCore/Public/ShaderParameters.h b/Engine/Source/Runtime/RenderCore/Public/ShaderParameters.h
--- a/Engine/Source/Runtime/RenderCore/Public/ShaderParameters.h
+++ b/Engine/Source/Runtime/RenderCore/Public/ShaderParameters.h
@@ -165,6 +165,23 @@ public:
 	}
 };
 
+/** A shader uniform buffer member binding, this is only used to determine if the member is used in the compiled shader. */
+class FShaderUniformBufferMemberParameter
+{
+	DECLARE_EXPORTED_TYPE_LAYOUT(FShaderUniformBufferMemberParameter, RENDERCORE_API, NonVirtual);
+public:
+	FShaderUniformBufferMemberParameter() = default;
+
+	RENDERCORE_API void Bind(const FShaderParameterMap& ParameterMap, const TCHAR* ParameterName);
+	friend RENDERCORE_API FArchive& operator<<(FArchive& Ar, FShaderUniformBufferMemberParameter& P);
+
+	inline bool IsBound() const { return bIsBound != 0; }
+	inline bool IsInitialized() const { return true; }
+
+private:
+	LAYOUT_FIELD_INITIALIZED(uint8, bIsBound, 0);
+};
+
 #if RHI_RAYTRACING
 
 struct UE_DEPRECATED(5.5, "Use FRHIBatchedShaderParameters and SetShaderParameters() instead.") FRayTracingShaderBindingsWriter : FRayTracingShaderBindings
diff --git a/Engine/Source/Runtime/Renderer/Private/MeshPassProcessor.cpp b/Engine/Source/Runtime/Renderer/Private/MeshPassProcessor.cpp
--- a/Engine/Source/Runtime/Renderer/Private/MeshPassProcessor.cpp
+++ b/Engine/Source/Runtime/Renderer/Private/MeshPassProcessor.cpp
@@ -839,7 +839,7 @@ void FMeshDrawShaderBindings::Finalize(const FMeshProcessorShaders* ShadersForDe
 				if (AutomaticallyBoundUniformBufferStruct)
 				{
 					ensureMsgf(
-						UniformBufferValue || EnumHasAnyFlags(AutomaticallyBoundUniformBufferStruct->GetBindingFlags(), EUniformBufferBindingFlags::Static),
+						UniformBufferValue || EnumHasAnyFlags(AutomaticallyBoundUniformBufferStruct->GetBindingFlags(), EUniformBufferBindingFlags::Static) || EnumHasAnyFlags(AutomaticallyBoundUniformBufferStruct->GetUsageFlags(), FShaderParametersMetadata::EUsageFlags::ManuallyBoundByPass),
 						TEXT("Shader %s with vertex factory %s never set automatically bound uniform buffer at BaseIndex %i.  Expected buffer of type %s.  This can cause GPU hangs, depending on how the shader uses it."),
 						Shader.GetType()->GetName(), 
 						VFType ? VFType->GetName() : TEXT("nullptr"),
@@ -1778,16 +1778,16 @@ public:
 		// This Layout fully replicates BatchedPrimitive UB
 		// we replace RDG_SRV with a regular SRV to be able to update UB inside RDG passes
 		static FName BatchedPrimitiveSlotName = "BatchedPrimitive";
-		FRHIUniformBufferLayoutInitializer Initialzer(TEXT("DynamicBatchedPrimitive"), 16u);
-		Initialzer.bUniformView = true;
-		Initialzer.Resources.Add({0, UBMT_RDG_BUFFER_SRV});
-		Initialzer.StaticSlot = FUniformBufferStaticSlotRegistry::Get().FindSlotByName(BatchedPrimitiveSlotName);
-		Initialzer.BindingFlags = EUniformBufferBindingFlags::StaticAndShader;
-		Initialzer.ComputeHash();
+		FRHIUniformBufferLayoutInitializer Initializer(TEXT("DynamicBatchedPrimitive"), 16u);
+		EnumAddFlags(Initializer.Flags, ERHIUniformBufferFlags::UniformView);
+		Initializer.Resources.Add({0, UBMT_RDG_BUFFER_SRV});
+		Initializer.StaticSlot = FUniformBufferStaticSlotRegistry::Get().FindSlotByName(BatchedPrimitiveSlotName);
+		Initializer.BindingFlags = EUniformBufferBindingFlags::StaticAndShader;
+		Initializer.ComputeHash();
 		// set view source to a regular SRV after hash computation, to make sure hash matches BatchedPrimitive layout
-		Initialzer.Resources[0].MemberType = UBMT_SRV;
+		Initializer.Resources[0].MemberType = UBMT_SRV;
 
-		LayoutRHI = RHICreateUniformBufferLayout(Initialzer);
+		LayoutRHI = RHICreateUniformBufferLayout(Initializer);
 	}
 
 	void ReleaseRHI() override 
diff --git a/Engine/Source/Runtime/VulkanRHI/Private/VulkanUniformBuffer.cpp b/Engine/Source/Runtime/VulkanRHI/Private/VulkanUniformBuffer.cpp
--- a/Engine/Source/Runtime/VulkanRHI/Private/VulkanUniformBuffer.cpp
+++ b/Engine/Source/Runtime/VulkanRHI/Private/VulkanUniformBuffer.cpp
@@ -146,7 +146,7 @@ FVulkanUniformBuffer::FVulkanUniformBuffer(FVulkanDevice& InDevice, const FRHIUn
 		}
 	}
 
-	if (InLayout->bUniformView)
+	if (EnumHasAnyFlags(InLayout->Flags, ERHIUniformBufferFlags::UniformView))
 	{
 		// For uniform view we expect an buffer SRV as a first resource
 		check(InLayout->Resources.Num() > 0);
