Fixed wrong FString constructor being used to construct from a range.

diff --git a/Engine/Source/Developer/ShaderCompilerCommon/Private/ShaderConductorContext.cpp b/Engine/Source/Developer/ShaderCompilerCommon/Private/ShaderConductorContext.cpp
--- a/Engine/Source/Developer/ShaderCompilerCommon/Private/ShaderConductorContext.cpp
+++ b/Engine/Source/Developer/ShaderCompilerCommon/Private/ShaderConductorContext.cpp
@@ -796,7 +796,7 @@ namespace CrossCompiler
 			[&OutSource](const void* Data, uint32 Size)
 			{
 				// Convert source buffer to FString
-				FString Converted(reinterpret_cast<const ANSICHAR*>(Data), Size);
+				FString Converted = FString::ConstructFromPtrSize(reinterpret_cast<const ANSICHAR*>(Data), Size);
 				OutSource = MoveTemp(Converted);
 			}
 		);
diff --git a/Engine/Source/Runtime/StorageServerClient/Private/StorageServerPlatformFile.cpp b/Engine/Source/Runtime/StorageServerClient/Private/StorageServerPlatformFile.cpp
--- a/Engine/Source/Runtime/StorageServerClient/Private/StorageServerPlatformFile.cpp
+++ b/Engine/Source/Runtime/StorageServerClient/Private/StorageServerPlatformFile.cpp
@@ -893,7 +893,7 @@ FString FStorageServerPlatformFile::ConvertToAbsolutePathForExternalAppForRead(c
 		Result << DotSlashSkip;
 		if (LowerLevel->FileExists(Result.ToString()))
 		{
-			return FString(Result.GetData(), Result.Len());
+			return FString::ConstructFromPtrSize(Result.GetData(), Result.Len());
 		}
 	}
 #endif
