PR #12150: [MetaXR]: Integrate Adreno Offline Shader Compiler into UE

#rnx

diff --git a/Engine/Source/Developer/ShaderCompilerCommon/Private/ShaderCompilerCommon.cpp b/Engine/Source/Developer/ShaderCompilerCommon/Private/ShaderCompilerCommon.cpp
--- a/Engine/Source/Developer/ShaderCompilerCommon/Private/ShaderCompilerCommon.cpp
+++ b/Engine/Source/Developer/ShaderCompilerCommon/Private/ShaderCompilerCommon.cpp
@@ -1563,24 +1563,25 @@ SHADERCOMPILERCOMMON_API void WriteShaderConductorCommandLine(const FShaderCompi
 	}
 }
 
-static uint32 Mali_ExtractNumberInstructions(const FString &MaliOutput)
+static uint32 OfflineCompiler_ExtractStats(const FString& CompilerOutput, const TArray<FString>& InstructionStrings)
 {
 	uint32 ReturnedNum = 0;
 
 	// Parse the instruction count
-	int32 InstructionStringLength = FPlatformString::Strlen(TEXT("Instructions Emitted:"));
-	int32 InstructionsIndex = MaliOutput.Find(TEXT("Instructions Emitted:"));
-
-	// new version of mali offline compiler uses a different string in its output
-	if (InstructionsIndex == INDEX_NONE)
+	int32 InstructionStringLength = 0, InstructionsIndex = 0;
+	for (const FString& InstrStr : InstructionStrings)
 	{
-		InstructionStringLength = FPlatformString::Strlen(TEXT("Total instruction cycles:"));
-		InstructionsIndex = MaliOutput.Find(TEXT("Total instruction cycles:"));
+		InstructionStringLength = InstrStr.Len();
+		InstructionsIndex = CompilerOutput.Find(*InstrStr);
+		if (InstructionsIndex != INDEX_NONE)
+		{
+			break;
+		}
 	}
 
-	if (InstructionsIndex != INDEX_NONE && InstructionsIndex + InstructionStringLength < MaliOutput.Len())
+	if (InstructionsIndex != INDEX_NONE && InstructionsIndex + InstructionStringLength < CompilerOutput.Len())
 	{
-		const int32 EndIndex = MaliOutput.Find(TEXT("\n"), ESearchCase::IgnoreCase, ESearchDir::FromStart, InstructionsIndex + InstructionStringLength);
+		const int32 EndIndex = CompilerOutput.Find(TEXT("\n"), ESearchCase::IgnoreCase, ESearchDir::FromStart, InstructionsIndex + InstructionStringLength);
 
 		if (EndIndex != INDEX_NONE)
 		{
@@ -1591,17 +1592,17 @@ static uint32 Mali_ExtractNumberInstructions(const FString &MaliOutput)
 
 			while (StartIndex < EndIndex)
 			{
-				if (FChar::IsDigit(MaliOutput[StartIndex]) && !bFoundNrStart)
+				if (FChar::IsDigit(CompilerOutput[StartIndex]) && !bFoundNrStart)
 				{
 					// found number's beginning
 					bFoundNrStart = true;
 					NumberIndex = StartIndex;
 				}
-				else if (FChar::IsWhitespace(MaliOutput[StartIndex]) && bFoundNrStart)
+				else if (FChar::IsWhitespace(CompilerOutput[StartIndex]) && bFoundNrStart)
 				{
 					// found number's end
 					bFoundNrStart = false;
-					const FString NumberString = MaliOutput.Mid(NumberIndex, StartIndex - NumberIndex);
+					const FString NumberString = CompilerOutput.Mid(NumberIndex, StartIndex - NumberIndex);
 					const float fNrInstructions = FCString::Atof(*NumberString);
 					ReturnedNum += (uint32)FMath::Max(0.0, ceil(fNrInstructions));
 				}
@@ -1614,24 +1615,24 @@ static uint32 Mali_ExtractNumberInstructions(const FString &MaliOutput)
 	return ReturnedNum;
 }
 
-static FString Mali_ExtractErrors(const FString &MaliOutput)
+static FString OfflineCompiler_ExtractErrors(const FString& CompilerOutput)
 {
 	FString ReturnedErrors;
 
-	const int32 GlobalErrorIndex = MaliOutput.Find(TEXT("Compilation failed."));
+	const int32 GlobalErrorIndex = CompilerOutput.Find(TEXT("Compilation failed."));
 
 	// find each 'line' that begins with token "ERROR:" and copy it to the returned string
 	if (GlobalErrorIndex != INDEX_NONE)
 	{
-		int32 CompilationErrorIndex = MaliOutput.Find(TEXT("ERROR:"));
+		int32 CompilationErrorIndex = CompilerOutput.Find(TEXT("ERROR:"));
 		while (CompilationErrorIndex != INDEX_NONE)
 		{
-			int32 EndLineIndex = MaliOutput.Find(TEXT("\n"), ESearchCase::CaseSensitive, ESearchDir::FromStart, CompilationErrorIndex + 1);
-			EndLineIndex = EndLineIndex == INDEX_NONE ? MaliOutput.Len() - 1 : EndLineIndex;
+			int32 EndLineIndex = CompilerOutput.Find(TEXT("\n"), ESearchCase::CaseSensitive, ESearchDir::FromStart, CompilationErrorIndex + 1);
+			EndLineIndex = EndLineIndex == INDEX_NONE ? CompilerOutput.Len() - 1 : EndLineIndex;
 
-			ReturnedErrors += MaliOutput.Mid(CompilationErrorIndex, EndLineIndex - CompilationErrorIndex + 1);
+			ReturnedErrors += CompilerOutput.Mid(CompilationErrorIndex, EndLineIndex - CompilationErrorIndex + 1);
 
-			CompilationErrorIndex = MaliOutput.Find(TEXT("ERROR:"), ESearchCase::CaseSensitive, ESearchDir::FromStart, EndLineIndex);
+			CompilationErrorIndex = CompilerOutput.Find(TEXT("ERROR:"), ESearchCase::CaseSensitive, ESearchDir::FromStart, EndLineIndex);
 		}
 	}
 
@@ -1640,58 +1641,142 @@ static FString Mali_ExtractErrors(const FString &MaliOutput)
 
 void CompileOfflineMali(const FShaderCompilerInput& Input, FShaderCompilerOutput& ShaderOutput, const ANSICHAR* ShaderSource, const int32 SourceSize, bool bVulkanSpirV, const ANSICHAR* VulkanSpirVEntryPoint)
 {
-	const bool bCompilerExecutableExists = FPaths::FileExists(Input.ExtraSettings.OfflineCompilerPath);
-
-	if (bCompilerExecutableExists)
-	{
-		const auto Frequency = (EShaderFrequency)Input.Target.Frequency;
-		const FString WorkingDir(FPlatformProcess::ShaderDir());
+	return CompileShaderOffline(Input, ShaderOutput, ShaderSource, SourceSize, bVulkanSpirV, VulkanSpirVEntryPoint);
+}
 
-		FString CompilerPath = Input.ExtraSettings.OfflineCompilerPath;
+/**
+* OfflineShaderCompiler's compilation command line options.
+* Each OfflineShaderCompiler should specify its own FOfflineShaderCompilerOptions. If one option is not supported by this OfflineShaderCompiler, leave it empty.
+* Frequency (VS/PS/etc.) here is called as stage sometime, too.
+*/
+class FOfflineShaderCompilerOptions
+{
+public:
+	/** Options applied to all shaders. */
+	FString CommonOptions;
+	/** MultiView option if it's enabled. */
+	FString MultiViewOption;
+	/** GPUTarget option*/
+	FString GPUTargetOption;
+	/** Default GPUTarget*/
+	FString DefaultGPUTarget;
+	/** Dump All*/
+	FString DumpAll;
+	/** SpirV file extension name*/
+	FString SpirVExt;
+	/** Default file extension name*/
+	FString DefaultGLSLExt;
+	/** GLSL source file extension used to specify which shader stage is being compiled. */
+	TMap<EShaderFrequency, FString> FrequencyGLSLExts;
+	/** Option to specify which shader stage is being compiled. */
+	TMap<EShaderFrequency, FString> FrequencyOptions;
+	/** Entrypoint option used to specify the entry point of each shader frequency. */
+	TMap<EShaderFrequency, FString> FrequencyEntryPoints;
+	/** Extra option of each shader frequency. */
+	TMap<EShaderFrequency, FString> FrequencyExtraOption;
+	/** Entrypoint option used to specify the entry point of all shader frequencies. */
+	TCHAR* DefaultEntryPoint;
+
+	/** Used to parse stats output to find total instruction count. Using array to support multiple compiler versions*/
+	TArray<FString> NumInstructionNames;
+	/** Used to parse stats output to find each stat. Each item of this array is for one stat AND it is also an array to support multiple compiler versions. */
+	TArray<TArray<FString>> StatsNames;
+
+	static FString GetFrequenceName(EShaderFrequency Freq)
+	{
+		static FString FrequencyName[SF_NumFrequencies] =
+		{
+			TEXT("VS"), // SF_Vertex
+			TEXT("MS"), // SF_Mesh
+			TEXT("AS"), // SF_Amplification
+			TEXT("FS"), // SF_Pixel
+			TEXT("GS"), // SF_Geometry
+			TEXT("CS")  // SF_Compute
+		};
+		if (Freq <= SF_Compute)
+		{
+			return FrequencyName[Freq];
+		}
+		else
+		{
+			return FString("Unknown");
+		}
+	}
+};
 
-		FString CompilerCommand = "";
+void CompileShaderOffline(const FShaderCompilerInput& Input,
+	FShaderCompilerOutput& ShaderOutput,
+	const ANSICHAR* ShaderSource,
+	const int32 SourceSize,
+	bool bVulkanSpirV,
+	const FOfflineShaderCompilerOptions& Options,
+	const ANSICHAR* VulkanSpirVEntryPoint)
+{
+	const auto Frequency = (EShaderFrequency)Input.Target.Frequency;
+	const FString WorkingDir(FPlatformProcess::ShaderDir());
 
-		// add process and thread ids to the file name to avoid collision between workers
-		auto ProcID = FPlatformProcess::GetCurrentProcessId();
-		auto ThreadID = FPlatformTLS::GetCurrentThreadId();
-		FString GLSLSourceFile = WorkingDir / TEXT("GLSLSource#") + FString::FromInt(ProcID) + TEXT("#") + FString::FromInt(ThreadID);
+	FString CompilerPath = Input.ExtraSettings.OfflineCompilerPath;
 
-		// setup compilation arguments
-		TCHAR *FileExt = nullptr;
-		switch (Frequency)
-		{
-			case SF_Vertex:
-				GLSLSourceFile += bVulkanSpirV ? TEXT(".spv") : TEXT(".vert");
-				CompilerCommand += TEXT(" -v");
-			break;
-			case SF_Pixel:
-				GLSLSourceFile += bVulkanSpirV ? TEXT(".spv") : TEXT(".frag");
-				CompilerCommand += TEXT(" -f");
-			break;
-			case SF_Geometry:
-				GLSLSourceFile += bVulkanSpirV ? TEXT(".spv") : TEXT(".geom");
-				CompilerCommand += TEXT(" -g");
-			break;
-			case SF_Compute:
-				GLSLSourceFile += bVulkanSpirV ? TEXT(".spv") : TEXT(".comp");
-				CompilerCommand += TEXT(" -C");
-			break;
+	// add process and thread ids to the file name to avoid collision between workers
+	auto ProcID = FPlatformProcess::GetCurrentProcessId();
+	auto ThreadID = FPlatformTLS::GetCurrentThreadId();
 
-			default:
-				GLSLSourceFile += TEXT(".shd");
-			break;
-		}
+	auto GetFileName = [&](FString FileType, FString Ext, int NumInst = 0)
+	{
+		return WorkingDir
+			/ FOfflineShaderCompilerOptions::GetFrequenceName(Frequency) + (NumInst ? TEXT("-") + FString::FromInt(NumInst) : TEXT(""))
+			+ FString::Printf(TEXT("-%s"), (VulkanSpirVEntryPoint ? ANSI_TO_TCHAR(VulkanSpirVEntryPoint) : TEXT("")))
+			+ *FileType + TEXT("-") + FString::FromInt(ProcID) + TEXT("-") + FString::FromInt(ThreadID)
+			+ *Ext;
+	};
 
-		if (bVulkanSpirV)
+	FString ShaderSrcExt;
+	if (bVulkanSpirV)
+	{
+		ShaderSrcExt = Options.SpirVExt;
+	}
+	else
+	{
+		if (const FString* Ext = Options.FrequencyGLSLExts.Find(Frequency))
 		{
-			CompilerCommand += FString::Printf(TEXT(" -y %s -p"), ANSI_TO_TCHAR(VulkanSpirVEntryPoint));
+			ShaderSrcExt = *Ext;
 		}
 		else
 		{
-			CompilerCommand += TEXT(" -s");
+			ShaderSrcExt = Options.DefaultGLSLExt;
+		}
+	}
+	
+	FString ShaderSourceFile = GetFileName(FString("-Source"), ShaderSrcExt);
+	FString CompilerCommand = Options.CommonOptions;
+	if (!Options.GPUTargetOption.IsEmpty())
+	{
+		FString GPUTarget = Input.ExtraSettings.GPUTarget;
+		if (GPUTarget.IsEmpty())
+		{
+			GPUTarget = Options.DefaultGPUTarget;			
 		}
+		CompilerCommand += FString::Printf(TEXT("%s=%s"), *Options.GPUTargetOption, *GPUTarget);
+	}
+	
+	if (Input.ExtraSettings.bMobileMultiView)
+	{
+		CompilerCommand += Options.MultiViewOption;
+	}
 
-		FArchive* Ar = IFileManager::Get().CreateFileWriter(*GLSLSourceFile, FILEWRITE_EvenIfReadOnly);
+	CompilerCommand += Options.FrequencyOptions[Frequency];
+
+	if (bVulkanSpirV)
+	{
+		CompilerCommand += FString::Printf(TEXT("%s %s"), *Options.FrequencyEntryPoints[Frequency], ANSI_TO_TCHAR(VulkanSpirVEntryPoint));
+	}
+
+	if (const FString* ExtraOption = Options.FrequencyExtraOption.Find(Frequency)) 
+	{
+		CompilerCommand += *ExtraOption;
+	}
+
+	FArchive* Ar = IFileManager::Get().CreateFileWriter(*ShaderSourceFile, FILEWRITE_EvenIfReadOnly);
 
 		if (Ar == nullptr)
 		{
@@ -1709,18 +1794,18 @@ void CompileOfflineMali(const FShaderCompilerInput& Input, FShaderCompilerOutput
 		// Since v6.2.0, Mali compiler needs to be started in the executable folder or it won't find "external/glslangValidator" for Vulkan
 		FString CompilerWorkingDirectory = FPaths::GetPath(CompilerPath);
 
-		if (!CompilerWorkingDirectory.IsEmpty() && FPaths::DirectoryExists(CompilerWorkingDirectory))
-		{
-			// compiler command line contains flags and the GLSL source file name
-			CompilerCommand += " " + FPaths::ConvertRelativePathToFull(GLSLSourceFile);
+	if (!CompilerWorkingDirectory.IsEmpty() && FPaths::DirectoryExists(CompilerWorkingDirectory))
+	{
+		// compiler command line contains flags and the GLSL source file name
+		CompilerCommand += " " + FPaths::ConvertRelativePathToFull(ShaderSourceFile);
 
-			// Run Mali shader compiler and wait for completion
-			FPlatformProcess::ExecProcess(*CompilerPath, *CompilerCommand, &ReturnCode, &StdOut, &StdErr, *CompilerWorkingDirectory);
-		}
-		else
-		{
-			StdErr = "Couldn't find Mali offline compiler at " + CompilerPath;
-		}
+		// Run shader compiler and wait for completion
+		FPlatformProcess::ExecProcess(*CompilerPath, *CompilerCommand, &ReturnCode, &StdOut, &StdErr, *CompilerWorkingDirectory);
+	}
+	else
+	{
+		StdErr = "Couldn't find offline compiler at " + CompilerPath;
+	}
 
 		// parse Mali's output and extract instruction count or eventual errors
 		ShaderOutput.bSucceeded = (ReturnCode >= 0);
@@ -1732,16 +1817,16 @@ void CompileOfflineMali(const FShaderCompilerInput& Input, FShaderCompilerOutput
 				ShaderOutput.bSucceeded = false;
 
 				FShaderCompilerError& NewError = ShaderOutput.Errors.AddDefaulted_GetRef();
-				NewError.StrippedErrorMessage = TEXT("[Mali Offline Complier]\n") + StdErr;
+				NewError.StrippedErrorMessage = TEXT("[Offline Complier]\n") + StdErr;
 			}
 			else
 			{
-				FString Errors = Mali_ExtractErrors(StdOut);
+				FString Errors = OfflineCompiler_ExtractErrors(StdOut);
 
 				if (Errors.Len())
 				{
 					FShaderCompilerError& NewError = ShaderOutput.Errors.AddDefaulted_GetRef();
-					NewError.StrippedErrorMessage = TEXT("[Mali Offline Complier]\n") + Errors;
+					NewError.StrippedErrorMessage = TEXT("[Offline Complier]\n") + Errors;
 					ShaderOutput.bSucceeded = false;
 				}
 			}
@@ -1749,12 +1834,188 @@ void CompileOfflineMali(const FShaderCompilerInput& Input, FShaderCompilerOutput
 			// extract instruction count
 			if (ShaderOutput.bSucceeded)
 			{
-				ShaderOutput.NumInstructions = Mali_ExtractNumberInstructions(StdOut);
+				ShaderOutput.NumInstructions = OfflineCompiler_ExtractStats(StdOut, Options.NumInstructionNames);
+				FString OutputStatsFile = GetFileName(FString("-Stats"), FString(".txt"), ShaderOutput.NumInstructions);
+				for (auto& StatNames : Options.StatsNames)
+				{
+					if(StatNames.Num())
+					{
+						ShaderOutput.AddStatistic<uint32>(*StatNames[0], OfflineCompiler_ExtractStats(StdOut, StatNames));
+					}
+				}
+				if (Input.ExtraSettings.bSaveCompilerStatsFiles)
+				{
+					FArchive* ArOutput = IFileManager::Get().CreateFileWriter(*OutputStatsFile, FILEWRITE_EvenIfReadOnly);
+					if (ArOutput == nullptr)
+					{
+						return;
+					}
+					if (!Options.DumpAll.IsEmpty())
+					{
+						CompilerCommand += Options.DumpAll;
+						//TODO: It's expensive to run the process twice. Better to run it once with DumpAll and parse the StdOut to get Stats.
+						// But to do that, we need to know the preserved keyword for Stats.
+						FPlatformProcess::ExecProcess(*CompilerPath, *CompilerCommand, &ReturnCode, &StdOut, &StdErr, *CompilerWorkingDirectory);
+					}
+					FString StatsOutput = CompilerCommand + FString("\n") + StdOut;
+					const int32 StatsLen = StatsOutput.Len();
+					TSharedPtr<ANSICHAR> ShaderStats = MakeShareable(new ANSICHAR[StatsLen + 1]);
+					FCStringAnsi::Strcpy(ShaderStats.Get(), StatsLen + 1, TCHAR_TO_ANSI(*StatsOutput));
+					ArOutput->Serialize((void*)ShaderStats.Get(), StatsLen);
+					delete ArOutput;
+				}
 			}
-		}
+	}
 
 		// we're done so delete the shader file
-		IFileManager::Get().Delete(*GLSLSourceFile, true, true);
+	if (Input.ExtraSettings.bSaveCompilerStatsFiles)
+	{
+		FString DstShaderSourceFile = GetFileName(FString("-Source"), ShaderSrcExt, ShaderOutput.NumInstructions);
+		IFileManager::Get().Move(*DstShaderSourceFile, *ShaderSourceFile, true, true);
+		IFileManager::Get().Delete(*ShaderSourceFile, true, true);
+	}
+	IFileManager::Get().Delete(*ShaderSourceFile, true, true);
+}
+
+void CompileShaderOffline_Mali(const FShaderCompilerInput& Input,
+	FShaderCompilerOutput& ShaderOutput,
+	const ANSICHAR* ShaderSource,
+	const int32 SourceSize,
+	bool bVulkanSpirV,
+	const ANSICHAR* VulkanSpirVEntryPoint)
+{
+	static FOfflineShaderCompilerOptions Options;
+	if (bVulkanSpirV)
+	{
+		Options.CommonOptions = TEXT(" -p");
+	}
+	else
+	{
+		Options.CommonOptions = TEXT(" -s");
+	}
+
+	if (Options.SpirVExt.IsEmpty())
+	{
+		Options.SpirVExt = TEXT(".spv");
+		Options.DefaultGLSLExt = TEXT(".shd");
+		Options.FrequencyGLSLExts.Emplace(SF_Vertex, TEXT(".vert"));
+		Options.FrequencyGLSLExts.Emplace(SF_Pixel, TEXT(".frag"));
+		Options.FrequencyGLSLExts.Emplace(SF_Geometry, TEXT(".geom"));
+		Options.FrequencyGLSLExts.Emplace(SF_Compute, TEXT(".comp"));
+
+		Options.FrequencyOptions.Emplace(SF_Vertex, FString(" -v"));
+		Options.FrequencyOptions.Emplace(SF_Pixel, FString(" -f"));
+		Options.FrequencyOptions.Emplace(SF_Geometry, FString(" -g"));
+		Options.FrequencyOptions.Emplace(SF_Compute, FString(" -C"));
+
+		Options.FrequencyEntryPoints.Emplace(SF_Vertex, TEXT(" -y"));
+		Options.FrequencyEntryPoints.Emplace(SF_Pixel, TEXT(" -y"));
+		Options.FrequencyEntryPoints.Emplace(SF_Geometry, TEXT(" -y"));
+		Options.FrequencyEntryPoints.Emplace(SF_Compute, TEXT(" -y"));
+
+		Options.NumInstructionNames.Add("Instructions Emitted:");
+		Options.NumInstructionNames.Add("Total instruction cycles:");
+	}
+
+	CompileShaderOffline(Input, ShaderOutput, ShaderSource, SourceSize, bVulkanSpirV, Options, VulkanSpirVEntryPoint);
+}
+
+void CompileShaderOffline_Adreno(const FShaderCompilerInput& Input,
+	FShaderCompilerOutput& ShaderOutput,
+	const ANSICHAR* ShaderSource,
+	const int32 SourceSize,
+	bool bVulkanSpirV,
+	const ANSICHAR* VulkanSpirVEntryPoint)
+{
+	static FOfflineShaderCompilerOptions Options;
+	if (bVulkanSpirV)
+	{
+		Options.CommonOptions = TEXT(" -api=Vulkan");
+	}
+
+	if (Options.MultiViewOption.IsEmpty())
+	{
+		Options.MultiViewOption = TEXT(" -view_mask=0x3");
+		Options.GPUTargetOption = TEXT(" -arch");
+		Options.DefaultGPUTarget = TEXT("a650");
+
+		Options.SpirVExt = TEXT(".spv");
+		Options.DefaultGLSLExt = TEXT(".shd");
+		Options.FrequencyGLSLExts.Emplace(SF_Vertex, TEXT(".vert"));
+		Options.FrequencyGLSLExts.Emplace(SF_Pixel, TEXT(".frag"));
+		Options.FrequencyGLSLExts.Emplace(SF_Geometry, TEXT(".geom"));
+		Options.FrequencyGLSLExts.Emplace(SF_Compute, TEXT(".comp"));
+
+		Options.FrequencyOptions.Emplace(SF_Vertex, FString(" -vs"));
+		Options.FrequencyOptions.Emplace(SF_Pixel, FString(" -fs"));
+		Options.FrequencyOptions.Emplace(SF_Geometry, FString(" -gs"));
+		Options.FrequencyOptions.Emplace(SF_Compute, FString(" -cs"));
+
+		Options.FrequencyEntryPoints.Emplace(SF_Vertex, TEXT(" -entry_point_vs"));
+		Options.FrequencyEntryPoints.Emplace(SF_Pixel, TEXT(" -entry_point_ps"));
+		Options.FrequencyEntryPoints.Emplace(SF_Geometry, TEXT(" -entry_point_gs"));
+		Options.FrequencyEntryPoints.Emplace(SF_Compute, TEXT(" -entry_point_cs"));
+
+		Options.FrequencyExtraOption.Emplace(SF_Vertex, TEXT(" -link_with_fs"));
+
+		Options.NumInstructionNames.Add("Total instruction count");
+		
+		Options.StatsNames.Add(TArray<FString>{TEXT("ALU instruction count - 32 bit")});
+		Options.StatsNames.Add(TArray<FString>{TEXT("ALU instruction count - 16 bit")});
+		Options.StatsNames.Add(TArray<FString>{TEXT("Complex instruction count - 32 bit")});
+		Options.StatsNames.Add(TArray<FString>{TEXT("Complex instruction count - 16 bit")});
+		Options.StatsNames.Add(TArray<FString>{TEXT("Flow control instruction count")});
+		Options.StatsNames.Add(TArray<FString>{TEXT("Barrier and fence Instruction count")});
+		Options.StatsNames.Add(TArray<FString>{TEXT("Short latency sync instruction count")});
+		Options.StatsNames.Add(TArray<FString>{TEXT("Long latency sync instruction count")});
+		Options.StatsNames.Add(TArray<FString>{TEXT("Texture read instruction count")});
+		Options.StatsNames.Add(TArray<FString>{TEXT("Memory read instruction count")});
+		Options.StatsNames.Add(TArray<FString>{TEXT("Memory write instruction count")});
+		Options.StatsNames.Add(TArray<FString>{TEXT("Miscellaneous instruction count")});
+		Options.StatsNames.Add(TArray<FString>{TEXT("Full precision register footprint per shader instance")});
+		Options.StatsNames.Add(TArray<FString>{TEXT("Half precision register footprint per shader instance")});
+		Options.StatsNames.Add(TArray<FString>{TEXT("Overall register footprint per shader instance")});
+		Options.StatsNames.Add(TArray<FString>{TEXT("Scratch memory usage per shader instance")});
+		Options.StatsNames.Add(TArray<FString>{TEXT("Loop count")});
+		Options.StatsNames.Add(TArray<FString>{TEXT("Output component count")});
+		Options.StatsNames.Add(TArray<FString>{TEXT("Input component count")});
+		Options.StatsNames.Add(TArray<FString>{TEXT("ALU fiber occupancy percentage")});
+		
+	}
+
+	if (Input.ExtraSettings.bDumpAll) 
+	{
+		Options.DumpAll = " -dump=all";
+	}
+
+	CompileShaderOffline(Input, ShaderOutput, ShaderSource, SourceSize, bVulkanSpirV, Options, VulkanSpirVEntryPoint);
+}
+
+void CompileShaderOffline(const FShaderCompilerInput& Input,
+	FShaderCompilerOutput& ShaderOutput,
+	const ANSICHAR* ShaderSource,
+	const int32 SourceSize,
+	bool bVulkanSpirV,
+	const ANSICHAR* VulkanSpirVEntryPoint)
+{
+	const bool bCompilerExecutableExists = FPaths::FileExists(Input.ExtraSettings.OfflineCompilerPath);
+	if (!bCompilerExecutableExists)
+	{
+		return;
+	}
+	EOfflineShaderCompilerType OfflineCompiler = Input.ExtraSettings.OfflineCompiler;
+	switch (OfflineCompiler)
+	{
+	case EOfflineShaderCompilerType::Mali:
+		CompileShaderOffline_Mali(Input, ShaderOutput, ShaderSource, SourceSize, bVulkanSpirV, VulkanSpirVEntryPoint);
+		break;
+	case EOfflineShaderCompilerType::Adreno:
+		CompileShaderOffline_Adreno(Input, ShaderOutput, ShaderSource, SourceSize, bVulkanSpirV, VulkanSpirVEntryPoint);
+		break;
+	case EOfflineShaderCompilerType::Num:
+		break;
+	default:
+		break;
 	}
 }
 
diff --git a/Engine/Source/Developer/ShaderCompilerCommon/Public/ShaderCompilerCommon.h b/Engine/Source/Developer/ShaderCompilerCommon/Public/ShaderCompilerCommon.h
--- a/Engine/Source/Developer/ShaderCompilerCommon/Public/ShaderCompilerCommon.h
+++ b/Engine/Source/Developer/ShaderCompilerCommon/Public/ShaderCompilerCommon.h
@@ -453,9 +453,17 @@ extern SHADERCOMPILERCOMMON_API void DumpDebugShaderBinary(const FShaderCompiler
 extern SHADERCOMPILERCOMMON_API void DumpDebugShaderDisassembledSpirv(const FShaderCompilerInput& Input, void* InData, int32 InDataByteSize, const FString& FileExtension);
 extern SHADERCOMPILERCOMMON_API void DumpDebugShaderDisassembledDxil(const FShaderCompilerInput& Input, void* InData, int32 InDataByteSize, const FString& FileExtension);
 
-// calls 'Mali Offline Compiler' to compile the glsl source code and extract the generated instruction count
+UE_DEPRECATED(5.5, "CompileOfflineMali is no longer used; CompileShaderOffline is used instead to support different compilers instead of Mali itself.")
 extern SHADERCOMPILERCOMMON_API void CompileOfflineMali(const FShaderCompilerInput &Input, FShaderCompilerOutput& ShaderOutput, const ANSICHAR* ShaderSource, const int32 SourceSize, bool bVulkanSpirV, const ANSICHAR* VulkanSpirVEntryPoint = nullptr);
 
+// calls 'Offline Compiler' to compile the source code and extract the stats
+extern SHADERCOMPILERCOMMON_API void CompileShaderOffline(const FShaderCompilerInput& Input,
+	FShaderCompilerOutput& ShaderOutput,
+	const ANSICHAR* ShaderSource,
+	const int32 SourceSize,
+	bool bVulkanSpirV,
+	const ANSICHAR* VulkanSpirVEntryPoint = nullptr);
+
 // Cross compiler support/common functionality
 namespace CrossCompiler
 {
diff --git a/Engine/Source/Developer/ShaderFormatOpenGL/Private/OpenGLShaderCompiler.cpp b/Engine/Source/Developer/ShaderFormatOpenGL/Private/OpenGLShaderCompiler.cpp
--- a/Engine/Source/Developer/ShaderFormatOpenGL/Private/OpenGLShaderCompiler.cpp
+++ b/Engine/Source/Developer/ShaderFormatOpenGL/Private/OpenGLShaderCompiler.cpp
@@ -3305,6 +3305,6 @@ static void PlatformCompileOfflineInternal(const FShaderCompilerInput& Input, FS
 {
 	if (ShaderVersion == GLSL_ES3_1_ANDROID)
 	{
-		CompileOfflineMali(Input, ShaderOutput, ShaderSource, FPlatformString::Strlen(ShaderSource), false);
+		CompileShaderOffline(Input, ShaderOutput, ShaderSource, FPlatformString::Strlen(ShaderSource), false);
 	}
 }
diff --git a/Engine/Source/Developer/VulkanShaderFormat/Private/SpirVShaderCompiler.inl b/Engine/Source/Developer/VulkanShaderFormat/Private/SpirVShaderCompiler.inl
--- a/Engine/Source/Developer/VulkanShaderFormat/Private/SpirVShaderCompiler.inl
+++ b/Engine/Source/Developer/VulkanShaderFormat/Private/SpirVShaderCompiler.inl
@@ -468,7 +468,7 @@ static void BuildShaderOutput(
 	{
 		if (InternalState.SupportsOfflineCompiler())
 		{
-			CompileOfflineMali(ShaderInput, ShaderOutput, (const ANSICHAR*)SerializedOutput.Spirv.GetByteData(), SerializedOutput.Spirv.GetByteSize(), true, SerializedOutput.SpirvEntryPointName);
+			CompileShaderOffline(ShaderInput, ShaderOutput, (const ANSICHAR*)SerializedOutput.Spirv.GetByteData(), SerializedOutput.Spirv.GetByteSize(), true, SerializedOutput.SpirvEntryPointName);
 		}
 	}
 
diff --git a/Engine/Source/Editor/MaterialEditor/Private/MaterialEditor.cpp b/Engine/Source/Editor/MaterialEditor/Private/MaterialEditor.cpp
--- a/Engine/Source/Editor/MaterialEditor/Private/MaterialEditor.cpp
+++ b/Engine/Source/Editor/MaterialEditor/Private/MaterialEditor.cpp
@@ -2426,7 +2426,7 @@ void FMaterialEditor::DrawMaterialInfoStrings(
 
 		for (int32 InstructionIndex = 0; InstructionIndex < Descriptions.Num(); InstructionIndex++)
 		{
-			FString InstructionCountString = FString::Printf(TEXT("%s: %u instructions"), *Descriptions[InstructionIndex].ShaderDescription, Descriptions[InstructionIndex].InstructionCount);
+			FString InstructionCountString = FString::Printf(TEXT("%s: %u instructions\nStats: %s"), *Descriptions[InstructionIndex].ShaderDescription, Descriptions[InstructionIndex].InstructionCount, *Descriptions[InstructionIndex].ShaderStatisticsString);
 			Canvas->DrawShadowedString(5, DrawPositionY, *InstructionCountString, FontToUse, FLinearColor(1, 1, 0));
 			DrawPositionY += SpacingBetweenLines;
 		}
@@ -3269,10 +3269,10 @@ void FMaterialEditor::UpdateMaterialinfoList_Old()
 
 			for (int32 InstructionIndex = 0; InstructionIndex < Results.Num(); InstructionIndex++)
 			{
-				FString InstructionCountString = FString::Printf(TEXT("%s: %u instructions"), *Results[InstructionIndex].ShaderDescription, Results[InstructionIndex].InstructionCount);
+					FString InstructionCountString = FString::Printf(TEXT("%s: %u instructions\nStats: %s"), *Results[InstructionIndex].ShaderDescription, Results[InstructionIndex].InstructionCount, *Results[InstructionIndex].ShaderStatisticsString);
 				if (bBuiltinStats)
 				{
-					InstructionCountString += FString::Printf(TEXT(" - Built-in instructions: %u"), EmptyMaterialResults[InstructionIndex].InstructionCount);
+						InstructionCountString += FString::Printf(TEXT(" - Built-in instructions: %u\nStats: %s"), EmptyMaterialResults[InstructionIndex].InstructionCount, *EmptyMaterialResults[InstructionIndex].ShaderStatisticsString);
 				}
 				TempMaterialInfoList.Add(MakeShareable(new FMaterialInfo(InstructionCountString, FLinearColor::Yellow)));
 				TSharedRef<FTokenizedMessage> Line = FTokenizedMessage::Create(EMessageSeverity::Info);
diff --git a/Engine/Source/Editor/MaterialEditor/Private/MaterialStats.cpp b/Engine/Source/Editor/MaterialEditor/Private/MaterialStats.cpp
--- a/Engine/Source/Editor/MaterialEditor/Private/MaterialStats.cpp
+++ b/Engine/Source/Editor/MaterialEditor/Private/MaterialStats.cpp
@@ -151,6 +151,34 @@ FText FShaderPlatformSettings::GetShaderCode(const EMaterialQualityLevel::Type Q
 	return Instance.ShaderCode;
 }
 
+/** returns all shaders' stats concatenated together*/
+FString FShaderPlatformSettings::GetShadersStats() const
+{
+	static FString ExtraLine("---------------------------------------------\n");
+	FString ShadersStats;
+	for (const FPlatformData& LevelData : PlatformData)
+	{
+		for (int32 ShaderType = 0; ShaderType < (int32)ERepresentativeShader::Num; ++ShaderType)
+		{
+			for (const FInstanceData& InstanceData : LevelData.Instances)
+			{
+				if (InstanceData.ShaderStatsInfo.ShaderInstructionCount.IsEmpty())
+				{
+					continue;
+				}
+				
+				if (auto* Count = InstanceData.ShaderStatsInfo.ShaderInstructionCount.Find((ERepresentativeShader)ShaderType))
+				{
+					ShadersStats += FString("\n") + ExtraLine + FMaterialStatsUtils::RepresentativeShaderTypeToString((ERepresentativeShader)ShaderType);
+					ShadersStats += Count->StrDescriptionLong;
+				}
+			}
+		}
+	}
+	
+	return ShadersStats;
+}
+
 void FShaderPlatformSettings::AllocateMaterialResources()
 {
 	ClearResources();
@@ -806,6 +834,26 @@ FText FMaterialStats::GetShaderCode(const EShaderPlatform PlatformID, const EMat
 	return (*Entry)->GetShaderCode(QualityType, InstanceIndex);
 }
 
+FString FMaterialStats::GetShadersStats() const
+{
+	static FString ExtraLine("=============================================\n");
+	FString ShadersStats;
+	int Index = 0; 
+	for (const auto& MapEntry : ShaderPlatformStatsDB)
+	{
+		TSharedPtr<FShaderPlatformSettings> PlatformPtr = MapEntry.Value;
+		FString PlatformShadersStats = PlatformPtr->GetShadersStats();
+		if (PlatformShadersStats.Len())
+		{
+			ShadersStats += (Index ? FString("\n\n") : FString("")) + ExtraLine + ExtraLine + PlatformPtr->GetPlatformName().GetPlainNameString();
+			ShadersStats += PlatformShadersStats;
+			++Index;
+		}
+	}
+
+	return ShadersStats;
+}
+
 void FMaterialStats::Update()
 {
 	const bool bNeedsUpdate = IsShowingStats() || IsCodeViewWindowActive();
diff --git a/Engine/Source/Editor/MaterialEditor/Private/MaterialStats.h b/Engine/Source/Editor/MaterialEditor/Private/MaterialStats.h
--- a/Engine/Source/Editor/MaterialEditor/Private/MaterialStats.h
+++ b/Engine/Source/Editor/MaterialEditor/Private/MaterialStats.h
@@ -277,6 +277,9 @@ public:
 	/** returns the actual shader source selected by the shaders viewer's combo-box */
 	FText GetShaderCode(const EMaterialQualityLevel::Type QualityType, const int32 InstanceIndex);
 
+	/** returns all shaders' stats concatenated together*/
+	FString GetShadersStats() const;
+
 	/** call this whenever the analyzed material or material instance is changed */
 	void SetMaterial(UMaterial *InBaseMaterial, UMaterialInstance *InBaseMaterialInstance, const TArray<TObjectPtr<UMaterialInstance>>& InDerivedMaterialInstances);
 
@@ -470,6 +473,9 @@ public:
 	/** returns the shader code computed by the specified platform with some quality level */
 	FText GetShaderCode(const EShaderPlatform PlatformID, const EMaterialQualityLevel::Type QualityType, const int32 InstanceIndex);
 
+	/** returns all shaders' stats concatenated together*/
+	FString GetShadersStats() const;
+
 	/** call this whenever some material property is changed, as it will trigger shader recompilation */
 	void SignalMaterialChanged();
 
diff --git a/Engine/Source/Editor/MaterialEditor/Private/MaterialStatsCommon.cpp b/Engine/Source/Editor/MaterialEditor/Private/MaterialStatsCommon.cpp
--- a/Engine/Source/Editor/MaterialEditor/Private/MaterialStatsCommon.cpp
+++ b/Engine/Source/Editor/MaterialEditor/Private/MaterialStatsCommon.cpp
@@ -18,7 +18,7 @@
 void FMaterialResourceStats::SetupExtraCompilationSettings(const EShaderPlatform Platform, FExtraShaderCompilerSettings& Settings) const
 {
 	Settings.bExtractShaderSource = true;
-	Settings.OfflineCompilerPath = FMaterialStatsUtils::GetPlatformOfflineCompilerPath(Platform);
+	FMaterialStatsUtils::GetPlatformOfflineCompilerSettings(Platform, Settings);
 }
 
 /*end FMaterialResourceStats functions*/
@@ -138,22 +138,52 @@ FString FMaterialStatsUtils::ShaderPlatformTypeName(const EShaderPlatform Platfo
 
 FString FMaterialStatsUtils::GetPlatformOfflineCompilerPath(const EShaderPlatform ShaderPlatform)
 {
+	FExtraShaderCompilerSettings SCSettings;
+	GetPlatformOfflineCompilerSettings(ShaderPlatform, SCSettings);
+	return SCSettings.OfflineCompilerPath;
+}
+
+void FMaterialStatsUtils::GetPlatformOfflineCompilerSettings(const EShaderPlatform ShaderPlatform, FExtraShaderCompilerSettings& SCSettings)
+{
+	auto GetSCType = [](EOfflineShaderCompiler SC)
+	{
+		switch (SC)
+		{
+		case EOfflineShaderCompiler::Mali:
+			return EOfflineShaderCompilerType::Mali;
+			break;
+		case EOfflineShaderCompiler::Adreno:
+			return EOfflineShaderCompilerType::Adreno;
+			break;
+		default:
+			return EOfflineShaderCompilerType::Num;
+		}
+	};
 	if (FDataDrivenShaderPlatformInfo::GetNeedsOfflineCompiler(ShaderPlatform))
 	{
 		if (FDataDrivenShaderPlatformInfo::GetIsAndroidOpenGLES(ShaderPlatform)
 			|| (FDataDrivenShaderPlatformInfo::GetIsLanguageVulkan(ShaderPlatform) && FDataDrivenShaderPlatformInfo::GetIsMobile(ShaderPlatform)))
 		{
-			return FPaths::ConvertRelativePathToFull(GetDefault<UMaterialEditorSettings>()->MaliOfflineCompilerPath.FilePath);
+			SCSettings.OfflineCompiler = GetSCType(GetDefault<UMaterialEditorSettings>()->OfflineCompiler);
+			SCSettings.OfflineCompilerPath = FPaths::ConvertRelativePathToFull(GetDefault<UMaterialEditorSettings>()->OfflineCompilerPath.FilePath);
+			SCSettings.GPUTarget = GetDefault<UMaterialEditorSettings>()->GPUTarget;
+			SCSettings.bDumpAll = GetDefault<UMaterialEditorSettings>()->bDumpAll;
+			SCSettings.bSaveCompilerStatsFiles = GetDefault<UMaterialEditorSettings>()->bSaveCompilerStatsFiles;
+			static const auto CVarMobileMultiView = IConsoleManager::Get().FindTConsoleVariableDataInt(TEXT("vr.MobileMultiView"));
+			SCSettings.bMobileMultiView = (CVarMobileMultiView && CVarMobileMultiView->GetValueOnAnyThread() != 0);
+			return;
 		}
 	}
-	return FString();
+	SCSettings.OfflineCompiler = EOfflineShaderCompilerType::Num;
+	return;
 }
 
 bool FMaterialStatsUtils::IsPlatformOfflineCompilerAvailable(const EShaderPlatform ShaderPlatform)
 {
-	FString CompilerPath = GetPlatformOfflineCompilerPath(ShaderPlatform);
+	FExtraShaderCompilerSettings Settings;
+	GetPlatformOfflineCompilerSettings(ShaderPlatform, Settings);
 
-	bool bCompilerExists = FPaths::FileExists(CompilerPath);
+	bool bCompilerExists = FPaths::FileExists(Settings.OfflineCompilerPath);
 
 	return bCompilerExists;
 }
@@ -686,7 +716,7 @@ void FMaterialStatsUtils::ExtractMatertialStatsInfo(EShaderPlatform ShaderPlatfo
 
 			Content.StrDescription = ShaderInstructionInfo[InstructionIndex].InstructionCount > 0 ? FString::Printf(TEXT("%u"), ShaderInstructionInfo[InstructionIndex].InstructionCount) : TEXT("n/a");
 			Content.StrDescriptionLong = ShaderInstructionInfo[InstructionIndex].InstructionCount > 0 ?
-				FString::Printf(TEXT("%s: %u instructions"), *ShaderInstructionInfo[InstructionIndex].ShaderDescription, ShaderInstructionInfo[InstructionIndex].InstructionCount) :
+				FString::Printf(TEXT("%s: %u instructions\nStats: %s"), *ShaderInstructionInfo[InstructionIndex].ShaderDescription, ShaderInstructionInfo[InstructionIndex].InstructionCount, *ShaderInstructionInfo[InstructionIndex].ShaderStatisticsString) :
 				TEXT("Offline shader compiler not available or an error was encountered!");
 
 			OutInfo.ShaderInstructionCount.Add(ShaderInstructionInfo[InstructionIndex].ShaderType, Content);
diff --git a/Engine/Source/Editor/MaterialEditor/Private/MaterialStatsGrid.cpp b/Engine/Source/Editor/MaterialEditor/Private/MaterialStatsGrid.cpp
--- a/Engine/Source/Editor/MaterialEditor/Private/MaterialStatsGrid.cpp
+++ b/Engine/Source/Editor/MaterialEditor/Private/MaterialStatsGrid.cpp
@@ -778,6 +778,11 @@ void FMaterialStatsGrid::BuildShaderRows()
 			VertShaderRow->CreateRow(StatsManager);
 
 			VertexShaderRows.Add(VertShaderRow);
+
+			TSharedPtr<FStatsGridRow> VertShaderRow2 = MakeShareable(new FStatsGridRow_Shaders((ERepresentativeShader)i, false, false));
+			VertShaderRow2->CreateRow(StatsManager);
+
+			VertexShaderRows.Add(VertShaderRow2);			
 		}
 	}
 }
diff --git a/Engine/Source/Editor/MaterialEditor/Private/SMaterialEditorStatsWidget.cpp b/Engine/Source/Editor/MaterialEditor/Private/SMaterialEditorStatsWidget.cpp
--- a/Engine/Source/Editor/MaterialEditor/Private/SMaterialEditorStatsWidget.cpp
+++ b/Engine/Source/Editor/MaterialEditor/Private/SMaterialEditorStatsWidget.cpp
@@ -6,6 +6,7 @@
 #include "Widgets/Layout/SExpandableArea.h"
 #include "Widgets/Layout/SScrollBox.h"
 #include "Widgets/Views/SListView.h"
+#include "Widgets/Input/SButton.h"
 #include "Widgets/Input/SCheckBox.h"
 #include "Widgets/Input/SComboButton.h"
 #include "Widgets/Input/SComboBox.h"
@@ -16,6 +17,7 @@
 #include "MaterialStatsGrid.h"
 #include "Framework/Application/SlateApplication.h"
 #include "Styling/StyleColors.h"
+#include "HAL/PlatformApplicationMisc.h"
 
 #define LOCTEXT_NAMESPACE "MaterialEditorStatsWidget"
 
@@ -639,6 +641,20 @@ void SMaterialEditorStatsWidget::ClearMessages()
 	MessageBoxWidget->ClearChildren();
 }
 
+FReply SMaterialEditorStatsWidget::OnButtonClick(EAppReturnType::Type ButtonID)
+{
+	auto MaterialStats = MaterialStatsWPtr.Pin();
+	if (!MaterialStats.IsValid())
+	{
+		return FReply::Handled();
+	}
+
+	FString ShadersStats = MaterialStats->GetShadersStats();
+	FPlatformApplicationMisc::ClipboardCopy(*ShadersStats);
+
+	return FReply::Handled();
+}
+
 void SMaterialEditorStatsWidget::Construct(const FArguments& InArgs)
 {
 	MaterialStatsWPtr = InArgs._MaterialStatsWPtr;
@@ -681,6 +697,16 @@ void SMaterialEditorStatsWidget::Construct(const FArguments& InArgs)
 			[
 				SNullWidget::NullWidget
 			]
+			+ SHorizontalBox::Slot()
+			.AutoWidth()
+			[
+				SNew(SButton)
+				.HAlign(HAlign_Center)
+				.ContentPadding(0)
+				.Text(LOCTEXT("CopyStats", "CopyStats"))
+				.ToolTipText(LOCTEXT("CopyStatsTooltip", "Copy all stats to the clipboard"))
+				.OnClicked(this, &SMaterialEditorStatsWidget::OnButtonClick, EAppReturnType::Cancel)
+			] 
 			+SHorizontalBox::Slot()
 			.AutoWidth()
 			[
diff --git a/Engine/Source/Editor/MaterialEditor/Private/SMaterialEditorStatsWidget.h b/Engine/Source/Editor/MaterialEditor/Private/SMaterialEditorStatsWidget.h
--- a/Engine/Source/Editor/MaterialEditor/Private/SMaterialEditorStatsWidget.h
+++ b/Engine/Source/Editor/MaterialEditor/Private/SMaterialEditorStatsWidget.h
@@ -72,6 +72,8 @@ private:
 
 	TSharedPtr<SWidget> BuildMessageArea();
 
+	FReply OnButtonClick(EAppReturnType::Type ButtonID);
+
 protected:
 	static const float ColumnSizeSmall;
 	static const float ColumnSizeMedium;
diff --git a/Engine/Source/Editor/MaterialEditor/Public/MaterialEditorSettings.h b/Engine/Source/Editor/MaterialEditor/Public/MaterialEditorSettings.h
--- a/Engine/Source/Editor/MaterialEditor/Public/MaterialEditorSettings.h
+++ b/Engine/Source/Editor/MaterialEditor/Public/MaterialEditorSettings.h
@@ -5,6 +5,16 @@
 #include "Engine/EngineTypes.h"
 #include "MaterialEditorSettings.generated.h"
 
+/**
+ * Enumerates offline shader compiler for the material editor
+ */
+UENUM()
+enum class EOfflineShaderCompiler : uint8
+{
+	Mali UMETA(DisplayName = "Mali Offline Compiler"),
+	Adreno UMETA(DisplayName = "Adreno Offline Compiler")
+};
+
 /**
  * Enumerates background for the material editor UI viewport (similar to ETextureEditorBackgrounds)
  */
@@ -96,12 +106,39 @@ class MATERIALEDITOR_API UMaterialEditorSettings : public UObject
 
 public:
 	/**
-	Path to user installed Mali shader compiler that can be used by the material editor to compile and extract shader informations for Android platforms.
-	Official website address: https://developer.arm.com/products/software-development-tools/graphics-development-tools/mali-offline-compiler/downloads
+	Offline shader compiler to use.
+	Mali Offline Compiler: Official website address: https://developer.arm.com/products/software-development-tools/graphics-development-tools/mali-offline-compiler/downloads
+	Adreno Offline Compiler: Official website address: https://developer.qualcomm.com/software/adreno-gpu-sdk/tools
 	*/
-	UPROPERTY(config, EditAnywhere, Category = "Offline Shader Compilers", meta = (DisplayName = "Mali Offline Compiler"))
+	UPROPERTY(config, EditAnywhere, Category = "Offline Shader Compilers", meta = (DisplayName = "Offline Shader Compiler"))
+	EOfflineShaderCompiler OfflineCompiler;
+
+	UE_DEPRECATED(5.5, "This property has been deprecated, please use OfflineCompilerPath instead.")
 	FFilePath MaliOfflineCompilerPath;
+	
+	/**
+	Path to user installed shader compiler that can be used by the material editor to compile and extract shader informations for Android platforms.
+	*/
+	UPROPERTY(config, EditAnywhere, Category = "Offline Shader Compilers", meta = (DisplayName = "Offline Compiler Path"))
+	FFilePath OfflineCompilerPath;
 
+	/**
+	The GPU target if the offline shader compiler needs one (Adreno GPU only).
+	*/
+	UPROPERTY(config, EditAnywhere, Category = "Offline Shader Compilers", meta = (DisplayName = "GPU Target", EditCondition = "OfflineCompiler == EOfflineShaderCompiler::Adreno"))
+	FString GPUTarget;
+
+	/**
+	Whether to save offline compiler stats files to Engine\Shaders
+	*/
+	UPROPERTY(config, EditAnywhere, Category = "Offline Shader Compilers", meta = (DisplayName = "Save Compiler Stats Files", EditCondition = "OfflineCompiler == EOfflineShaderCompiler::Adreno"))
+	bool bSaveCompilerStatsFiles = false;
+
+	/**
+	Whether to dump stats only or all information to file (Adreno GPU only).
+	*/
+	UPROPERTY(config, EditAnywhere, Category = "Offline Shader Compilers", meta = (DisplayName = "Dump All Stats To Compiler Stats Files", EditCondition = "OfflineCompiler == EOfflineShaderCompiler::Adreno"))
+	bool bDumpAll;
 protected:
 	// The width (in pixels) of the preview viewport when a material editor is first opened
 	UPROPERTY(config, EditAnywhere, meta=(ClampMin=1, ClampMax=4096), Category="User Interface Domain")
diff --git a/Engine/Source/Editor/MaterialEditor/Public/MaterialStatsCommon.h b/Engine/Source/Editor/MaterialEditor/Public/MaterialStatsCommon.h
--- a/Engine/Source/Editor/MaterialEditor/Public/MaterialStatsCommon.h
+++ b/Engine/Source/Editor/MaterialEditor/Public/MaterialStatsCommon.h
@@ -120,6 +120,8 @@ public:
 	static FSlateColor PlatformTypeColor(EPlatformCategoryType PlatformType);
 
 	MATERIALEDITOR_API static bool IsPlatformOfflineCompilerAvailable(const EShaderPlatform ShaderPlatform);
+	UE_DEPRECATED(5.5, "GetPlatformOfflineCompilerPath is no longer used; Use GetPlatformOfflineCompilerSettings instead.")
 	MATERIALEDITOR_API static FString GetPlatformOfflineCompilerPath(const EShaderPlatform ShaderPlatform);
+	MATERIALEDITOR_API static void GetPlatformOfflineCompilerSettings(const EShaderPlatform ShaderPlatform, FExtraShaderCompilerSettings& SCSettings);
 	MATERIALEDITOR_API static bool PlatformNeedsOfflineCompiler(const EShaderPlatform ShaderPlatform);
 };
diff --git a/Engine/Source/Runtime/RenderCore/Public/ShaderCompilerCore.h b/Engine/Source/Runtime/RenderCore/Public/ShaderCompilerCore.h
--- a/Engine/Source/Runtime/RenderCore/Public/ShaderCompilerCore.h
+++ b/Engine/Source/Runtime/RenderCore/Public/ShaderCompilerCore.h
@@ -201,16 +201,31 @@ struct FShaderCompilerResourceTable
 	}
 };
 
+/** enumeration of offline shader compiler for the material editor */
+enum class EOfflineShaderCompilerType : uint8
+{
+	Mali,
+	Adreno,
+
+	Num
+};
+
 /** Additional compilation settings that can be configured by each FMaterial instance before compilation */
 struct FExtraShaderCompilerSettings
 {
 	bool bExtractShaderSource = false;
 	FString OfflineCompilerPath;
+	EOfflineShaderCompilerType OfflineCompiler = EOfflineShaderCompilerType::Mali;
+	FString GPUTarget;
+	bool bDumpAll = false;
+	bool bSaveCompilerStatsFiles = false;
+	bool bMobileMultiView = false;
 
 	friend FArchive& operator<<(FArchive& Ar, FExtraShaderCompilerSettings& StatsSettings)
 	{
 		// Note: this serialize is used to pass between UE and the shader compile worker, recompile both when modifying
-		return Ar << StatsSettings.bExtractShaderSource << StatsSettings.OfflineCompilerPath;
+		return Ar << StatsSettings.bExtractShaderSource << StatsSettings.OfflineCompilerPath 
+			<< StatsSettings.OfflineCompiler << StatsSettings.GPUTarget << StatsSettings.bDumpAll << StatsSettings.bSaveCompilerStatsFiles << StatsSettings.bMobileMultiView;
 	}
 };
 
