Use min16float precision in mobile preview

diff --git a/Engine/Shaders/Private/MobileBasePassCommon.ush b/Engine/Shaders/Private/MobileBasePassCommon.ush
--- a/Engine/Shaders/Private/MobileBasePassCommon.ush
+++ b/Engine/Shaders/Private/MobileBasePassCommon.ush
@@ -34,7 +34,11 @@
 struct FSharedMobileBasePassInterpolants
 {
 #if USE_VERTEX_FOG && !PACK_INTERPOLANTS
-	half4 VertexFog			: TEXCOORD7;
+	#if MOBILE_EMULATION
+		float4 VertexFog : TEXCOORD7; // half precision interpolators are causing issues in FXC fallback path as it's used in mobile preview
+	#else
+		half4 VertexFog : TEXCOORD7;
+	#endif
 #endif
 
 	float4 PixelPosition	: TEXCOORD8; // xyz = world position, w = clip z
diff --git a/Engine/Shaders/Public/Platform.ush b/Engine/Shaders/Public/Platform.ush
--- a/Engine/Shaders/Public/Platform.ush
+++ b/Engine/Shaders/Public/Platform.ush
@@ -240,8 +240,6 @@
 
 #ifndef PLATFORM_SUPPORTS_REAL_TYPES
 #define PLATFORM_SUPPORTS_REAL_TYPES 0
-#elif FORCE_FLOATS
-#error Collision between FORCE_FLOATS and PLATFORM_SUPPORTS_REAL_TYPES
 #endif
 
 #ifndef COMPILER_SUPPORTS_WAVE_32_64_MODE
@@ -324,10 +322,19 @@
 #endif
 
 #if COMPILER_DXC && !FORCE_FLOATS && (ES3_1_PROFILE || METAL_PROFILE)
-    #pragma dxc diagnostic ignored "-Wconversion"
+	#pragma dxc diagnostic ignored "-Wconversion"
 #endif
 
-#if ((!(ES3_1_PROFILE || METAL_PROFILE)) && !PLATFORM_SUPPORTS_REAL_TYPES) || FORCE_FLOATS
+#if ES3_1_PROFILE && MOBILE_EMULATION && !FORCE_FLOATS
+	#define half min16float
+	#define half1 min16float1
+	#define half2 min16float2
+	#define half3 min16float3
+	#define half4 min16float4
+	#define half3x3 min16float3x3
+	#define half3x4 min16float3x4
+	#define half4x4 min16float4x4
+#elif ((!(ES3_1_PROFILE || METAL_PROFILE)) && !PLATFORM_SUPPORTS_REAL_TYPES) || FORCE_FLOATS
 	// Always use floats when using the ES3/METAL compiler, because platforms not optimized for lower precision, 
 	// And we don't want potential side effects on other platforms
 	#define half float
@@ -480,11 +487,11 @@
 	#define UB_INT(Dim)         int##Dim
 	#define UB_UINT(Dim)        uint##Dim
 	#define UB_FLOAT(Dim)       float##Dim
-    #if METAL_PROFILE
-        #define UB_HALF_FLOAT(Dim)  float##Dim
-    #else
-        #define UB_HALF_FLOAT(Dim)  half##Dim
-    #endif
+	#if METAL_PROFILE || COMPILER_HLSL
+		#define UB_HALF_FLOAT(Dim)  float##Dim
+	#else
+		#define UB_HALF_FLOAT(Dim)  half##Dim
+	#endif
 	#define UB_FIXED_FLOAT(Dim) fixed##Dim
 #else
 	#define UB_INT(Dim)         int##Dim
diff --git a/Engine/Source/Developer/ShaderCompilerCommon/Private/HlslExpressionParser.inl b/Engine/Source/Developer/ShaderCompilerCommon/Private/HlslExpressionParser.inl
--- a/Engine/Source/Developer/ShaderCompilerCommon/Private/HlslExpressionParser.inl
+++ b/Engine/Source/Developer/ShaderCompilerCommon/Private/HlslExpressionParser.inl
@@ -323,6 +323,28 @@ namespace CrossCompiler
 		case EHlslToken::Half4x3:
 		case EHlslToken::Half4x4:
 
+		case EHlslToken::Min16Float:
+		case EHlslToken::Min16Float1:
+		case EHlslToken::Min16Float2:
+		case EHlslToken::Min16Float3:
+		case EHlslToken::Min16Float4:
+		case EHlslToken::Min16Float1x1:
+		case EHlslToken::Min16Float1x2:
+		case EHlslToken::Min16Float1x3:
+		case EHlslToken::Min16Float1x4:
+		case EHlslToken::Min16Float2x1:
+		case EHlslToken::Min16Float2x2:
+		case EHlslToken::Min16Float2x3:
+		case EHlslToken::Min16Float2x4:
+		case EHlslToken::Min16Float3x1:
+		case EHlslToken::Min16Float3x2:
+		case EHlslToken::Min16Float3x3:
+		case EHlslToken::Min16Float3x4:
+		case EHlslToken::Min16Float4x1:
+		case EHlslToken::Min16Float4x2:
+		case EHlslToken::Min16Float4x3:
+		case EHlslToken::Min16Float4x4:
+
 		case EHlslToken::Float:
 		case EHlslToken::Float1:
 		case EHlslToken::Float2:
diff --git a/Engine/Source/Developer/ShaderCompilerCommon/Private/HlslLexer.cpp b/Engine/Source/Developer/ShaderCompilerCommon/Private/HlslLexer.cpp
--- a/Engine/Source/Developer/ShaderCompilerCommon/Private/HlslLexer.cpp
+++ b/Engine/Source/Developer/ShaderCompilerCommon/Private/HlslLexer.cpp
@@ -338,6 +338,28 @@ namespace CrossCompiler
 				InsertToken(TEXT("half3x4"), EHlslToken::Half3x4);
 				InsertToken(TEXT("half4x4"), EHlslToken::Half4x4);
 
+				InsertToken(TEXT("min16float"), EHlslToken::Min16Float);
+				InsertToken(TEXT("min16float1"), EHlslToken::Min16Float1);
+				InsertToken(TEXT("min16float2"), EHlslToken::Min16Float2);
+				InsertToken(TEXT("min16float3"), EHlslToken::Min16Float3);
+				InsertToken(TEXT("min16float4"), EHlslToken::Min16Float4);
+				InsertToken(TEXT("min16float1x1"), EHlslToken::Min16Float1x1);
+				InsertToken(TEXT("min16float2x1"), EHlslToken::Min16Float2x1);
+				InsertToken(TEXT("min16float3x1"), EHlslToken::Min16Float3x1);
+				InsertToken(TEXT("min16float4x1"), EHlslToken::Min16Float4x1);
+				InsertToken(TEXT("min16float1x2"), EHlslToken::Min16Float1x2);
+				InsertToken(TEXT("min16float2x2"), EHlslToken::Min16Float2x2);
+				InsertToken(TEXT("min16float3x2"), EHlslToken::Min16Float3x2);
+				InsertToken(TEXT("min16float4x2"), EHlslToken::Min16Float4x2);
+				InsertToken(TEXT("min16float1x3"), EHlslToken::Min16Float1x3);
+				InsertToken(TEXT("min16float2x3"), EHlslToken::Min16Float2x3);
+				InsertToken(TEXT("min16float3x3"), EHlslToken::Min16Float3x3);
+				InsertToken(TEXT("min16float4x3"), EHlslToken::Min16Float4x3);
+				InsertToken(TEXT("min16float1x4"), EHlslToken::Min16Float1x4);
+				InsertToken(TEXT("min16float2x4"), EHlslToken::Min16Float2x4);
+				InsertToken(TEXT("min16float3x4"), EHlslToken::Min16Float3x4);
+				InsertToken(TEXT("min16float4x4"), EHlslToken::Min16Float4x4);
+
 				InsertToken(TEXT("float"), EHlslToken::Float);
 				InsertToken(TEXT("float1"), EHlslToken::Float1);
 				InsertToken(TEXT("float2"), EHlslToken::Float2);
diff --git a/Engine/Source/Developer/ShaderCompilerCommon/Private/HlslLexer.h b/Engine/Source/Developer/ShaderCompilerCommon/Private/HlslLexer.h
--- a/Engine/Source/Developer/ShaderCompilerCommon/Private/HlslLexer.h
+++ b/Engine/Source/Developer/ShaderCompilerCommon/Private/HlslLexer.h
@@ -195,6 +195,28 @@ namespace CrossCompiler
 		Half3x4,
 		Half4x4,
 
+		Min16Float,
+		Min16Float1,
+		Min16Float2,
+		Min16Float3,
+		Min16Float4,
+		Min16Float1x1,
+		Min16Float2x1,
+		Min16Float3x1,
+		Min16Float4x1,
+		Min16Float1x2,
+		Min16Float2x2,
+		Min16Float3x2,
+		Min16Float4x2,
+		Min16Float1x3,
+		Min16Float2x3,
+		Min16Float3x3,
+		Min16Float4x3,
+		Min16Float1x4,
+		Min16Float2x4,
+		Min16Float3x4,
+		Min16Float4x4,
+
 		Float,
 		Float1,
 		Float2,
diff --git a/Engine/Source/Developer/Windows/ShaderFormatD3D/Private/ShaderFormatD3D.cpp b/Engine/Source/Developer/Windows/ShaderFormatD3D/Private/ShaderFormatD3D.cpp
--- a/Engine/Source/Developer/Windows/ShaderFormatD3D/Private/ShaderFormatD3D.cpp
+++ b/Engine/Source/Developer/Windows/ShaderFormatD3D/Private/ShaderFormatD3D.cpp
@@ -258,6 +258,12 @@ public:
 			AddShaderTargetDefines(Input, 6, 8);
 			break;
 		}
+
+		// For mobile emulation
+		if (Input.Environment.FullPrecisionInPS || (Input.SharedEnvironment.IsValid() && Input.SharedEnvironment->FullPrecisionInPS))
+		{
+			Input.Environment.SetDefine(TEXT("FORCE_FLOATS"), (uint32)1);
+		}
 	}
 
 	virtual const TCHAR* GetPlatformIncludeDirectory() const
