[shaders] stripping out shader system-related code that was deprecated in 5.4

diff --git a/Engine/Source/Developer/Apple/MetalShaderFormat/Private/MetalCompileShaderSPIRV.cpp b/Engine/Source/Developer/Apple/MetalShaderFormat/Private/MetalCompileShaderSPIRV.cpp
--- a/Engine/Source/Developer/Apple/MetalShaderFormat/Private/MetalCompileShaderSPIRV.cpp
+++ b/Engine/Source/Developer/Apple/MetalShaderFormat/Private/MetalCompileShaderSPIRV.cpp
@@ -931,8 +931,6 @@ void FMetalCompileShaderSPIRV::DoCompileMetalShader(
 
 		CrossCompiler::FShaderConductorTarget TargetDesc;
 
-		PRAGMA_DISABLE_DEPRECATION_WARNINGS		// FShaderCompilerDefinitions will be made internal in the future, marked deprecated until then
-
 		if (Result)
 		{
 			SideTableIndex = FPlatformMath::CountTrailingZeros64(BufferIndices);
@@ -1007,8 +1005,6 @@ void FMetalCompileShaderSPIRV::DoCompileMetalShader(
 				TargetDesc.CompileFlags->SetDefine(TEXT("argument_buffer_offset"), IABOffsetIndex);
 			}
 			TargetDesc.CompileFlags->SetDefine(TEXT("texture_buffer_native"), 1);
-
-			PRAGMA_ENABLE_DEPRECATION_WARNINGS
 			
 			switch (VersionEnum)
 			{
diff --git a/Engine/Source/Developer/ShaderCompilerCommon/Private/ShaderCompilerCommon.cpp b/Engine/Source/Developer/ShaderCompilerCommon/Private/ShaderCompilerCommon.cpp
--- a/Engine/Source/Developer/ShaderCompilerCommon/Private/ShaderCompilerCommon.cpp
+++ b/Engine/Source/Developer/ShaderCompilerCommon/Private/ShaderCompilerCommon.cpp
@@ -13,6 +13,7 @@
 #include "Serialization/MemoryReader.h"
 #include "Serialization/MemoryWriter.h"
 #include "String/RemoveFrom.h"
+#include "ShaderCompilerDefinitions.h"
 #include "ShaderPreprocessor.h"
 #include "ShaderPreprocessTypes.h"
 #include "ShaderSymbolExport.h"
@@ -2047,9 +2048,7 @@ namespace UE::ShaderCompilerCommon
 		FShaderPreprocessOutput& PreprocessOutput,
 		const FShaderCompilerInput& Input,
 		const FShaderCompilerEnvironment& Environment,
-		PRAGMA_DISABLE_DEPRECATION_WARNINGS
 		const FShaderCompilerDefinitions& AdditionalDefines
-		PRAGMA_ENABLE_DEPRECATION_WARNINGS
 		)
 	{
 		TRACE_CPUPROFILER_EVENT_SCOPE(FBaseShaderFormat_PreprocessShader);
@@ -2092,6 +2091,16 @@ namespace UE::ShaderCompilerCommon
 		return bSuccess;
 	}
 
+	bool ExecuteShaderPreprocessingSteps(
+		FShaderPreprocessOutput& PreprocessOutput,
+		const FShaderCompilerInput& Input,
+		const FShaderCompilerEnvironment& Environment
+	)
+	{
+		// overloaded function rather than defaulting definitions parameter to avoid including internal header in public header
+		return ExecuteShaderPreprocessingSteps(PreprocessOutput, Input, Environment, FShaderCompilerDefinitions());
+	}
+
 	FString FDebugShaderDataOptions::GetDebugShaderPath(const FShaderCompilerInput& Input, const TCHAR* Suffix) const
 	{
 		return GetDebugFileName(Input, *this, OverrideBaseFilename, Suffix);
@@ -2375,172 +2384,6 @@ void DumpDebugShaderDisassembledDxil(const FShaderCompilerInput& Input, void* In
 
 namespace CrossCompiler
 {
-	FString CreateResourceTableFromEnvironment(const FShaderCompilerEnvironment& Environment)
-	{
-		FString Line = TEXT("\n#if 0 /*BEGIN_RESOURCE_TABLES*/\n");
-		for (auto Pair : Environment.UniformBufferMap)
-		{
-			Line += FString::Printf(TEXT("%s, %d\n"), *Pair.Key, Pair.Value.LayoutHash);
-		}
-		Line += TEXT("NULL, 0\n");
-		for (const FUniformResourceEntry& Entry : Environment.ResourceTableMap.Resources)
-		{
-			Line += FString::Printf(TEXT("%s, %s, %d, %d\n"), Entry.UniformBufferMemberName, *FString(Entry.GetUniformBufferName()), Entry.Type, Entry.ResourceIndex);
-		}
-		Line += TEXT("NULL, NULL, 0, 0\n");
-
-		Line += TEXT("#endif /*END_RESOURCE_TABLES*/\n");
-		return Line;
-	}
-
-	void CreateEnvironmentFromResourceTable(const FString& String, FShaderCompilerEnvironment& OutEnvironment)
-	{
-		FString Prolog = TEXT("#if 0 /*BEGIN_RESOURCE_TABLES*/");
-		int32 FoundBegin = String.Find(Prolog, ESearchCase::CaseSensitive);
-		if (FoundBegin == INDEX_NONE)
-		{
-			return;
-		}
-		int32 FoundEnd = String.Find(TEXT("#endif /*END_RESOURCE_TABLES*/"), ESearchCase::CaseSensitive, ESearchDir::FromStart, FoundBegin);
-		if (FoundEnd == INDEX_NONE)
-		{
-			return;
-		}
-
-		// +1 for EOL
-		const TCHAR* Ptr = &String[FoundBegin + 1 + Prolog.Len()];
-		while (*Ptr == '\r' || *Ptr == '\n')
-		{
-			++Ptr;
-		}
-		const TCHAR* PtrEnd = &String[FoundEnd];
-		while (Ptr < PtrEnd)
-		{
-			FString UB;
-			if (!CrossCompiler::ParseIdentifier(Ptr, UB))
-			{
-				return;
-			}
-			if (!CrossCompiler::Match(Ptr, TEXT(", ")))
-			{
-				return;
-			}
-			int32 Hash;
-			if (!CrossCompiler::ParseSignedNumber(Ptr, Hash))
-			{
-				return;
-			}
-			// Optional \r
-			CrossCompiler::Match(Ptr, '\r');
-			if (!CrossCompiler::Match(Ptr, '\n'))
-			{
-				return;
-			}
-
-			if (UB == TEXT("NULL") && Hash == 0)
-			{
-				break;
-			}
-
-			FUniformBufferEntry& UniformBufferEntry = OutEnvironment.UniformBufferMap.FindOrAdd(UB);
-			UniformBufferEntry.LayoutHash = (uint32)Hash;
-
-			if (!UniformBufferEntry.MemberNameBuffer)
-			{
-				TArray<TCHAR>* MemberNameBuffer = new TArray<TCHAR>();
-				UniformBufferEntry.MemberNameBuffer = MakeShareable(MemberNameBuffer);
-			}
-		}
-
-		// Need to iterate through Uniform Buffer Map to add strings to correct MemberNameBuffer storage
-		auto UniformBufferMapIt = OutEnvironment.UniformBufferMap.begin();
-		
-		// If we exit parse early due to error, we still want to fixup the string names for the members we found,
-		// so the partial data isn't corrupt.
-		struct FFixupOnExit
-		{
-			FFixupOnExit(FShaderCompilerEnvironment& OutEnvironment) : Environment(OutEnvironment) {}
-			~FFixupOnExit()
-			{
-				Environment.ResourceTableMap.FixupOnLoad(Environment.UniformBufferMap);
-			}
-			
-			FShaderCompilerEnvironment& Environment;
-		};
-		FFixupOnExit FixupOnExit(OutEnvironment);
-
-		while (Ptr < PtrEnd)
-		{
-			FString Name;
-			if (!CrossCompiler::ParseIdentifier(Ptr, Name))
-			{
-				return;
-			}
-			if (!CrossCompiler::Match(Ptr, TEXT(", ")))
-			{
-				return;
-			}
-			FString UB;
-			if (!CrossCompiler::ParseIdentifier(Ptr, UB))
-			{
-				return;
-			}
-			if (!CrossCompiler::Match(Ptr, TEXT(", ")))
-			{
-				return;
-			}
-			int32 Type;
-			if (!CrossCompiler::ParseSignedNumber(Ptr, Type))
-			{
-				return;
-			}
-			if (!CrossCompiler::Match(Ptr, TEXT(", ")))
-			{
-				return;
-			}
-			int32 ResourceIndex;
-			if (!CrossCompiler::ParseSignedNumber(Ptr, ResourceIndex))
-			{
-				return;
-			}
-			// Optional
-			CrossCompiler::Match(Ptr, '\r');
-			if (!CrossCompiler::Match(Ptr, '\n'))
-			{
-				return;
-			}
-
-			if (Name == TEXT("NULL") && UB == TEXT("NULL") && Type == 0 && ResourceIndex == 0)
-			{
-				break;
-			}
-
-			// Advance the uniform buffer map if this is a different UB name
-			while (UniformBufferMapIt.Key() != UB)
-			{
-				++UniformBufferMapIt;
-				if (UniformBufferMapIt == OutEnvironment.UniformBufferMap.end())
-				{
-					return;
-				}
-			}
-
-			// Append the Name we parsed to the member name buffer
-			TArray<TCHAR>& Buffer = *UniformBufferMapIt.Value().MemberNameBuffer.Get();
-			uint32 MemberNameLength = Name.Len();
-			
-			Buffer.Append(*Name, MemberNameLength + 1);
-
-			// The member name field of the entries is initialized at the end of parsing by the FixupOnLoad call from FFixupOnExit, so we can set it to nullptr here
-			OutEnvironment.ResourceTableMap.Resources.Add({
-				nullptr,
-				(uint8)UB.Len(),
-				(uint8)Type,
-				(uint16)ResourceIndex
-			});
-		}
-	}
-
 	/**
 	 * Parse an error emitted by the HLSL cross-compiler.
 	 * @param OutErrors - Array into which compiler errors may be added.
diff --git a/Engine/Source/Developer/ShaderCompilerCommon/Private/ShaderConductorContext.cpp b/Engine/Source/Developer/ShaderCompilerCommon/Private/ShaderConductorContext.cpp
--- a/Engine/Source/Developer/ShaderCompilerCommon/Private/ShaderConductorContext.cpp
+++ b/Engine/Source/Developer/ShaderCompilerCommon/Private/ShaderConductorContext.cpp
@@ -233,8 +233,6 @@ namespace CrossCompiler
 		}
 	}
 
-	PRAGMA_DISABLE_DEPRECATION_WARNINGS		// FShaderCompilerDefinitions will be made internal in the future, marked deprecated until then
-
 	// Converts a map of string pairs to a C-Style macro defines array
 	static void ConvertDefineMapToMacroDefines(const FShaderCompilerDefinitions& Definitions, TArray<TPair<FAnsiString, FAnsiString>>& OutPairs, TArray<ShaderConductor::MacroDefine>& OutPairRefs)
 	{
@@ -252,8 +250,7 @@ namespace CrossCompiler
 			OutPairRefs[Index].value = *OutPairs[Index].Value;
 		}
 	}
-	PRAGMA_ENABLE_DEPRECATION_WARNINGS
-
+	
 	static void ConvertScTargetDesc(FShaderConductorContext::FShaderConductorIntermediates& Intermediates, const FShaderConductorTarget& InTarget, ShaderConductor::Compiler::TargetDesc& OutTargetDesc)
 	{
 		// Convert FString to ANSI string and store them as intermediates
@@ -561,9 +558,7 @@ namespace CrossCompiler
 
 	FShaderConductorTarget::FShaderConductorTarget()
 	{
-		PRAGMA_DISABLE_DEPRECATION_WARNINGS		// FShaderCompilerDefinitions will be made internal in the future, marked deprecated until then
 		CompileFlags = MakePimpl<FShaderCompilerDefinitions>();
-		PRAGMA_ENABLE_DEPRECATION_WARNINGS
 	}
 
 	FShaderConductorContext::FShaderConductorContext()
@@ -592,9 +587,7 @@ namespace CrossCompiler
 		return *this;
 	}
 
-	PRAGMA_DISABLE_DEPRECATION_WARNINGS		// FShaderCompilerDefinitions will be made internal in the future, marked deprecated until then
 	bool FShaderConductorContext::LoadSource(FStringView ShaderSource, const FString& Filename, const FString& EntryPoint, EShaderFrequency ShaderStage, const FShaderCompilerDefinitions* Definitions, const TArray<FString>* ExtraDxcArgs)
-	PRAGMA_ENABLE_DEPRECATION_WARNINGS
 	{
 		// Convert FString to ANSI string and store them as intermediates
 		Intermediates->ShaderSource = ShaderSource;
@@ -618,9 +611,7 @@ namespace CrossCompiler
 		return true;
 	}
 
-	PRAGMA_DISABLE_DEPRECATION_WARNINGS		// FShaderCompilerDefinitions will be made internal in the future, marked deprecated until then
 	bool FShaderConductorContext::LoadSource(FAnsiStringView ShaderSource, const FString& Filename, const FString& EntryPoint, EShaderFrequency ShaderStage, const FShaderCompilerDefinitions* Definitions, const TArray<FString>* ExtraDxcArgs)
-	PRAGMA_ENABLE_DEPRECATION_WARNINGS
 	{
 		Intermediates->ShaderSource = ShaderSource;
 		Intermediates->Filename = Filename;
@@ -644,16 +635,12 @@ namespace CrossCompiler
 	}
 
 
-	PRAGMA_DISABLE_DEPRECATION_WARNINGS		// FShaderCompilerDefinitions will be made internal in the future, marked deprecated until then
 	bool FShaderConductorContext::LoadSource(const FString& ShaderSource, const FString& Filename, const FString& EntryPoint, EShaderFrequency ShaderStage, const FShaderCompilerDefinitions* Definitions, const TArray<FString>* ExtraDxcArgs)
-	PRAGMA_ENABLE_DEPRECATION_WARNINGS
 	{
 		return LoadSource(FStringView(ShaderSource), Filename, EntryPoint, ShaderStage, Definitions, ExtraDxcArgs);
 	}
 
-	PRAGMA_DISABLE_DEPRECATION_WARNINGS		// FShaderCompilerDefinitions will be made internal in the future, marked deprecated until then
 	bool FShaderConductorContext::LoadSource(const ANSICHAR* ShaderSource, const ANSICHAR* Filename, const ANSICHAR* EntryPoint, EShaderFrequency ShaderStage, const FShaderCompilerDefinitions* Definitions, const TArray<FString>* ExtraDxcArgs)
-	PRAGMA_ENABLE_DEPRECATION_WARNINGS
 	{
 		Intermediates->ShaderSource = ShaderSource;
 		Intermediates->Filename = Filename;
diff --git a/Engine/Source/Developer/ShaderCompilerCommon/Public/ShaderCompilerCommon.h b/Engine/Source/Developer/ShaderCompilerCommon/Public/ShaderCompilerCommon.h
--- a/Engine/Source/Developer/ShaderCompilerCommon/Public/ShaderCompilerCommon.h
+++ b/Engine/Source/Developer/ShaderCompilerCommon/Public/ShaderCompilerCommon.h
@@ -61,9 +61,13 @@ namespace UE::ShaderCompilerCommon
 		FShaderPreprocessOutput& PreprocessOutput,
 		const FShaderCompilerInput& Input,
 		const FShaderCompilerEnvironment& Environment,
-		PRAGMA_DISABLE_DEPRECATION_WARNINGS		// FShaderCompilerDefinitions will be made internal in the future, marked deprecated until then
-		const FShaderCompilerDefinitions& AdditionalDefines = FShaderCompilerDefinitions()
-		PRAGMA_ENABLE_DEPRECATION_WARNINGS
+		const FShaderCompilerDefinitions& AdditionalDefines
+		);
+
+	extern SHADERCOMPILERCOMMON_API bool ExecuteShaderPreprocessingSteps(
+		FShaderPreprocessOutput& PreprocessOutput,
+		const FShaderCompilerInput& Input,
+		const FShaderCompilerEnvironment& Environment
 		);
 
 	extern SHADERCOMPILERCOMMON_API FStringView          RemoveConstantBufferPrefix(FStringView InName);
@@ -95,13 +99,6 @@ namespace UE::ShaderCompilerCommon
 	extern SHADERCOMPILERCOMMON_API bool RemoveDeadCode(FShaderSource& InOutPreprocessedShaderSource, const FString& EntryPoint, TConstArrayView<FStringView> RequiredSymbols, TArray<FShaderCompilerError>& OutErrors);
 	extern SHADERCOMPILERCOMMON_API bool RemoveDeadCode(FShaderSource& InOutPreprocessedShaderSource, TConstArrayView<FStringView> RequiredSymbols, TArray<FShaderCompilerError>& OutErrors);
 	
-	UE_DEPRECATED(5.4, "Use overload of RemoveDeadCode accepting FShaderSource")
-	inline bool RemoveDeadCode(FString& InOutPreprocessedShaderSource, const FString& EntryPoint, TArray<FShaderCompilerError>& OutErrors) { return false; }
-	UE_DEPRECATED(5.4, "Use overload of RemoveDeadCode accepting FShaderSource")
-	inline bool RemoveDeadCode(FString& InOutPreprocessedShaderSource, const FString& EntryPoint, TConstArrayView<FStringView> RequiredSymbols, TArray<FShaderCompilerError>& OutErrors) { return false; }
-	UE_DEPRECATED(5.4, "Use overload of RemoveDeadCode accepting FShaderSource")
-	inline bool RemoveDeadCode(FString& InOutPreprocessedShaderSource, TConstArrayView<FStringView> RequiredSymbols, TArray<FShaderCompilerError>& OutErrors) { return false; }
-
 	struct FDebugShaderDataOptions
 	{
 		struct FAdditionalOutput
@@ -109,17 +106,7 @@ namespace UE::ShaderCompilerCommon
 			const TCHAR* BaseFileName;
 			const TCHAR* Data;
 		};
-		
-		PRAGMA_DISABLE_DEPRECATION_WARNINGS
-		// Explicitly-defaulted ctors are needed temporarily due to deprecation of the HlslCCFlags field.
-		// These can be removed once the deprecation window for said field ends.
-		FDebugShaderDataOptions() = default;
-		FDebugShaderDataOptions(FDebugShaderDataOptions&&) = default;
-		FDebugShaderDataOptions(const FDebugShaderDataOptions&) = default;
-		PRAGMA_ENABLE_DEPRECATION_WARNINGS
-		
-		UE_DEPRECATED(5.4, "HlslCCFlags field is no longer used in debug output; please remove any usage.")
-		uint32 HlslCCFlags = 0;
+
 		const TCHAR* OverrideBaseFilename = nullptr;
 		const TCHAR* FilenamePrefix = nullptr;
 		TFunction<FString()> AppendPreSource{};
@@ -171,12 +158,6 @@ namespace UE::ShaderCompilerCommon
 	 */
 	extern SHADERCOMPILERCOMMON_API FString GetDebugShaderContents(const FShaderCompilerInput& Input, FStringView PreprocessedSource, const FDebugShaderDataOptions& Options = FDebugShaderDataOptions(), const TCHAR* Suffix = nullptr);
 	
-	UE_DEPRECATED(5.4, "Use overload of GetDebugShaderContents accepting an FStringView")
-	inline FString GetDebugShaderContents(const FShaderCompilerInput& Input, const FString& PreprocessedSource, const FDebugShaderDataOptions& Options = FDebugShaderDataOptions())
-	{
-		return FString();
-	}
-
 	class FBaseShaderFormat : public IShaderFormat
 	{
 	public:
@@ -347,9 +328,6 @@ extern SHADERCOMPILERCOMMON_API void AddUnboundShaderParameterError(
 // Convert generated UniformBuffer code and references into something the shader compilers can use.
 extern SHADERCOMPILERCOMMON_API void CleanupUniformBufferCode(const FShaderCompilerEnvironment& Environment, FShaderSource& PreprocessedShaderSource);
 
-UE_DEPRECATED(5.4, "RemoveUniformBuffersFromSource was renamed CleanupUniformBufferCode")
-inline void RemoveUniformBuffersFromSource(const FShaderCompilerEnvironment& Environment, FString& PreprocessedShaderSource) {}
-
 template <typename CharType>
 const CharType* FindMatchingBlock(const CharType* OpeningCharPtr, char OpenChar, char CloseChar)
 {
@@ -384,9 +362,6 @@ extern SHADERCOMPILERCOMMON_API const TCHAR* ParseHLSLSymbolName(const TCHAR* Se
 extern SHADERCOMPILERCOMMON_API void ParseHLSLTypeName(const TCHAR* SearchString, const TCHAR*& TypeNameStartPtr, const TCHAR*& TypeNameEndPtr);
 extern SHADERCOMPILERCOMMON_API FStringView FindNextHLSLDefinitionOfType(FStringView Typename, FStringView StartPos);
 
-UE_DEPRECATED(5.4, "TEXT macro processing is now handled in PreprocessShader --  TransformStringIntoCharacterArray is a no-op")
-inline void TransformStringIntoCharacterArray(FString& PreprocessedShaderSource, TArray<FShaderDiagnosticData>* OutDiagnosticDatas = nullptr) {}
-
 // Structure to hold forward declarations for a specific scope/namespace chain for the HlslParser
 struct FScopedDeclarations
 {
@@ -467,11 +442,6 @@ extern SHADERCOMPILERCOMMON_API void CompileShaderOffline(const FShaderCompilerI
 // Cross compiler support/common functionality
 namespace CrossCompiler
 {
-	UE_DEPRECATED(5.4, "CreateResourceTableFromEnvironment is no longer used; serializing environment for SCW directcompile mode now uses a base64-encoded string containing all environment compilation dependencies.")
-	extern SHADERCOMPILERCOMMON_API FString CreateResourceTableFromEnvironment(const FShaderCompilerEnvironment& Environment);
-	UE_DEPRECATED(5.4, "CreateEnvironmentFromResourceTable is no longer used; serializing environment for SCW directcompile mode now uses a base64-encoded string containing all environment compilation dependencies.")
-	extern SHADERCOMPILERCOMMON_API void CreateEnvironmentFromResourceTable(const FString& String, FShaderCompilerEnvironment& OutEnvironment);
-
 	extern SHADERCOMPILERCOMMON_API void ParseHlslccError(TArray<FShaderCompilerError>& OutErrors, const FString& InLine, bool bUseAbsolutePaths = false);
 
 	struct SHADERCOMPILERCOMMON_API FHlslccHeader
diff --git a/Engine/Source/Developer/ShaderCompilerCommon/Public/ShaderConductorContext.h b/Engine/Source/Developer/ShaderCompilerCommon/Public/ShaderConductorContext.h
--- a/Engine/Source/Developer/ShaderCompilerCommon/Public/ShaderConductorContext.h
+++ b/Engine/Source/Developer/ShaderCompilerCommon/Public/ShaderConductorContext.h
@@ -155,9 +155,7 @@ namespace CrossCompiler
 		int32 Version = 0;
 
 		/** Cross compilation flags. This is used for high-level cross compilation (such as Metal output) that is send over to SPIRV-Cross, e.g. { "invariant_float_math", "1" }. */
-		PRAGMA_DISABLE_DEPRECATION_WARNINGS		// FShaderCompilerDefinitions will be made internal in the future, marked deprecated until then
 		TPimplPtr<FShaderCompilerDefinitions> CompileFlags;
-		PRAGMA_ENABLE_DEPRECATION_WARNINGS
 
 		/** Optional callback to rename certain variable types. */
 		TFunction<bool(const FAnsiStringView& VariableName, const FAnsiStringView& TypeName, FString& OutRenamedTypeName)> VariableTypeRenameCallback;
diff --git a/Engine/Source/Developer/ShaderCompilerCommon/Public/ShaderMinifier.h b/Engine/Source/Developer/ShaderCompilerCommon/Public/ShaderMinifier.h
--- a/Engine/Source/Developer/ShaderCompilerCommon/Public/ShaderMinifier.h
+++ b/Engine/Source/Developer/ShaderCompilerCommon/Public/ShaderMinifier.h
@@ -57,12 +57,5 @@ struct FMinifiedShader
 
 extern SHADERCOMPILERCOMMON_API FMinifiedShader Minify(const FShaderSource& PreprocessedShader, const FShaderSource::FViewType EntryPoint, EMinifyShaderFlags Flags = EMinifyShaderFlags::None);
 extern SHADERCOMPILERCOMMON_API FMinifiedShader Minify(const FShaderSource& PreprocessedShader, TConstArrayView<FShaderSource::FViewType> RequiredSymbols, EMinifyShaderFlags Flags = EMinifyShaderFlags::None);
-
-UE_DEPRECATED(5.4, "Use Minify overloads accepting FShaderSource")
-inline FMinifiedShader Minify(const FStringView PreprocessedShader, const FStringView EntryPoint, EMinifyShaderFlags Flags = EMinifyShaderFlags::None) { return FMinifiedShader(); }
-
-UE_DEPRECATED(5.4, "Use Minify overloads accepting FShaderSource")
-inline FMinifiedShader Minify(const FStringView PreprocessedShader, TConstArrayView<FStringView> RequiredSymbols, EMinifyShaderFlags Flags = EMinifyShaderFlags::None) { return FMinifiedShader(); }
-
 } // UE::ShaderMinifier
 
diff --git a/Engine/Source/Developer/ShaderFormatOpenGL/Private/OpenGLShaderCompiler.cpp b/Engine/Source/Developer/ShaderFormatOpenGL/Private/OpenGLShaderCompiler.cpp
--- a/Engine/Source/Developer/ShaderFormatOpenGL/Private/OpenGLShaderCompiler.cpp
+++ b/Engine/Source/Developer/ShaderFormatOpenGL/Private/OpenGLShaderCompiler.cpp
@@ -2654,8 +2654,6 @@ bool GenerateDeferredMobileShaders(std::string& GlslSource, GLSLCompileParameter
 		OutString += ESVersionString + "\n";
 		OutString += "#ifdef UE_MRT_FRAMEBUFFER_FETCH\n";
 
-		PRAGMA_DISABLE_DEPRECATION_WARNINGS		// FShaderCompilerDefinitions will be made internal in the future, marked deprecated until then
-
 		FShaderCompilerDefinitions CompileFlagsCopy = *GLSLCompileParams.TargetDesc->CompileFlags;
 
 		// Add defines for FBF in spirv-glsl
@@ -2683,8 +2681,6 @@ bool GenerateDeferredMobileShaders(std::string& GlslSource, GLSLCompileParameter
 
 		*GLSLCompileParams.TargetDesc->CompileFlags = CompileFlagsCopy;
 
-		PRAGMA_ENABLE_DEPRECATION_WARNINGS
-
 		OutString += FBFSourceString;
 
 		// Generate PLS shader
@@ -2738,9 +2734,7 @@ bool GenerateDeferredMobileShaders(std::string& GlslSource, GLSLCompileParameter
 					DefineValue = GBufferOutputNames[i];
 				}
 
-				PRAGMA_DISABLE_DEPRECATION_WARNINGS		// FShaderCompilerDefinitions will be made internal in the future, marked deprecated until then
 				GLSLCompileParams.TargetDesc->CompileFlags->SetDefine(*DefineKey, *DefineValue);
-				PRAGMA_ENABLE_DEPRECATION_WARNINGS
 			}
 		}
 
@@ -2853,14 +2847,10 @@ static bool CompileToGlslWithShaderConductor(
 		SourceData = HlslFrameBufferDeclarations + SourceData;
 	}
 	
-	PRAGMA_DISABLE_DEPRECATION_WARNINGS		// FShaderCompilerDefinitions will be made internal in the future, marked deprecated until then
-
 	// Inject additional macro definitions to circumvent missing features: external textures
 	FShaderCompilerDefinitions AdditionalDefines;
 	AdditionalDefines.SetDefine(TEXT("TextureExternal"), TEXT("Texture2D"));
 
-	PRAGMA_ENABLE_DEPRECATION_WARNINGS
-
 	uint32_t BlendFlags = GetDecalBlendFlags(Input);
 
 	// Load shader source into compiler context
@@ -2927,8 +2917,6 @@ static bool CompileToGlslWithShaderConductor(
 
 		CrossCompiler::FShaderConductorTarget TargetDesc;
 
-		PRAGMA_DISABLE_DEPRECATION_WARNINGS		// FShaderCompilerDefinitions will be made internal in the future, marked deprecated until then
-
 		TargetDesc.CompileFlags->SetDefine(TEXT("relax_nan_checks"), 1);
 
 		switch (Version)
@@ -2958,8 +2946,6 @@ static bool CompileToGlslWithShaderConductor(
 			break;
 		}
 
-		PRAGMA_ENABLE_DEPRECATION_WARNINGS
-
 		TSet<FString> ExternalTextures;
 		int32 Pos = 0;
 		static constexpr FStringView TextExternal = TEXTVIEW("TextureExternal");
diff --git a/Engine/Source/Developer/ShaderFormatVectorVM/Public/ShaderFormatVectorVM.h b/Engine/Source/Developer/ShaderFormatVectorVM/Public/ShaderFormatVectorVM.h
--- a/Engine/Source/Developer/ShaderFormatVectorVM/Public/ShaderFormatVectorVM.h
+++ b/Engine/Source/Developer/ShaderFormatVectorVM/Public/ShaderFormatVectorVM.h
@@ -63,9 +63,3 @@ inline FArchive& operator<<(FArchive& Ar, FVectorVMCompilationOutput& Output)
 	Ar << Output.Errors;
 	return Ar;
 }
-
-UE_DEPRECATED(5.4, "Direct VectorVM compilation will no longer be exposed outside of the TestCompileVectorVMShader function")
-inline bool CompileShader_VectorVM(const FShaderCompilerInput& Input, FShaderCompilerOutput& Output, const FString& WorkingDirectory, uint8 Version) { return false; }
-
-UE_DEPRECATED(5.4, "This overload of CompileShader_VectorVM is superceded by TestCompileShaderVectorVM")
-inline bool CompileShader_VectorVM(const FShaderCompilerInput& Input, FShaderCompilerOutput& Output, const FString& WorkingDirectory, uint8 Version, FVectorVMCompilationOutput& VMCompilationOutput, bool bSkipBackendOptimizations = false) { return false; }
diff --git a/Engine/Source/Developer/ShaderPreprocessor/Private/ShaderPreprocessor.cpp b/Engine/Source/Developer/ShaderPreprocessor/Private/ShaderPreprocessor.cpp
--- a/Engine/Source/Developer/ShaderPreprocessor/Private/ShaderPreprocessor.cpp
+++ b/Engine/Source/Developer/ShaderPreprocessor/Private/ShaderPreprocessor.cpp
@@ -68,7 +68,6 @@ struct FCompareResultPath
 	}
 };
 
-PRAGMA_DISABLE_DEPRECATION_WARNINGS		// FShaderCompilerDefinitions will be made internal in the future, marked deprecated until then
 static void AddStbDefine(stb_arena* MacroArena, macro_definition**& StbDefines, const TCHAR* Name, const TCHAR* Value);
 static void AddStbDefines(stb_arena* MacroArena, macro_definition**& StbDefines, const FShaderCompilerDefinitions& Defines);
 
@@ -82,7 +81,6 @@ public:
 		AddStbDefines(MacroArena, OutDefines, AdditionalDefines);
 	}
 };
-PRAGMA_ENABLE_DEPRECATION_WARNINGS
 
 //////////////////////////////////////////////////////////////////////////
 extern "C"
@@ -674,8 +672,6 @@ static void AddStbDefine(stb_arena* MacroArena, macro_definition**& StbDefines,
 	arrput(StbDefines, pp_define(MacroArena, Define.ToString()));
 }
 
-PRAGMA_DISABLE_DEPRECATION_WARNINGS		// FShaderCompilerDefinitions will be made internal in the future, marked deprecated until then
-
 static void AddStbDefines(stb_arena* MacroArena, macro_definition**& StbDefines, const FShaderCompilerDefinitions& Defines)
 {
 	for (FShaderCompilerDefinitions::FConstIterator It(Defines); It; ++It)
@@ -684,8 +680,6 @@ static void AddStbDefines(stb_arena* MacroArena, macro_definition**& StbDefines,
 	}
 }
 
-PRAGMA_ENABLE_DEPRECATION_WARNINGS
-
 /**
  * Preprocess a shader.
  * @param OutPreprocessedShader - Upon return contains the preprocessed source code.
@@ -699,9 +693,7 @@ bool PreprocessShader(
 	FShaderPreprocessOutput& Output,
 	const FShaderCompilerInput& Input,
 	const FShaderCompilerEnvironment& Environment,
-	PRAGMA_DISABLE_DEPRECATION_WARNINGS		// FShaderCompilerDefinitions will be made internal in the future, marked deprecated until then
 	const FShaderCompilerDefinitions& AdditionalDefines
-	PRAGMA_ENABLE_DEPRECATION_WARNINGS
 )
 {
 	TRACE_CPUPROFILER_EVENT_SCOPE(PreprocessShader);
@@ -709,8 +701,6 @@ bool PreprocessShader(
 	stb_arena MacroArena = { 0 };
 	macro_definition** StbDefines = nullptr;
 	
-	PRAGMA_DISABLE_DEPRECATION_WARNINGS
-	
 	FShaderPreprocessorUtilities::PopulateDefines(Environment, AdditionalDefines, &MacroArena, StbDefines);
 
 	// The substitution text generated by custom macros gets run through the preprocessor afterwards, but in some cases we want to
@@ -722,8 +712,6 @@ bool PreprocessShader(
 	arrput(StbDefines, pp_define_custom_macro(&MacroArena, ShaderPrintTextIdentifier, 1));
 	arrput(StbDefines, pp_define_custom_macro(&MacroArena, ShaderPrintAssertIdentifier, 0));
 
-	PRAGMA_ENABLE_DEPRECATION_WARNINGS
-
 	FStbPreprocessContext Context{ Input, Environment };
 
 	auto InFilename = StringCast<ANSICHAR>(*Input.VirtualSourceFilePath);
@@ -851,25 +839,10 @@ bool PreprocessShader(
 }
 
 bool PreprocessShader(
-	FString& OutPreprocessedShader,
-	FShaderCompilerOutput& ShaderOutput,
-	const FShaderCompilerInput& ShaderInput,
-	PRAGMA_DISABLE_DEPRECATION_WARNINGS		// FShaderCompilerDefinitions will be made internal in the future, marked deprecated until then
-	const FShaderCompilerDefinitions& AdditionalDefines,
-	EDumpShaderDefines DefinesPolicy
-	PRAGMA_ENABLE_DEPRECATION_WARNINGS
-)
+	FShaderPreprocessOutput& Output,
+	const FShaderCompilerInput& Input,
+	const FShaderCompilerEnvironment& MergedEnvironment)
 {
-	FShaderPreprocessOutput Output;
-	// when called via this overload, environment is assumed to be already merged in input struct
-	const FShaderCompilerEnvironment& Environment = ShaderInput.Environment;
-	bool bSucceeded = PreprocessShader(Output, ShaderInput, Environment, AdditionalDefines);
-
-	OutPreprocessedShader = FString(Output.GetSourceViewWide());
-
-	for (FShaderCompilerError& Error : Output.EditErrors())
-	{
-		ShaderOutput.Errors.Add(MoveTemp(Error));
-	}
-	return bSucceeded;
-}
+	// overload instead of defaulting FShaderCompilerDefinitions arg to avoid including an internal header in a public header.
+	return PreprocessShader(Output, Input, MergedEnvironment, FShaderCompilerDefinitions());
+}
\ No newline at end of file
diff --git a/Engine/Source/Developer/ShaderPreprocessor/Public/ShaderPreprocessor.h b/Engine/Source/Developer/ShaderPreprocessor/Public/ShaderPreprocessor.h
--- a/Engine/Source/Developer/ShaderPreprocessor/Public/ShaderPreprocessor.h
+++ b/Engine/Source/Developer/ShaderPreprocessor/Public/ShaderPreprocessor.h
@@ -12,29 +12,6 @@ class FString;
 struct FShaderCompilerInput;
 struct FShaderCompilerOutput;
 
-enum class 
-	UE_DEPRECATED(5.4, "EDumpShaderDefines is no longer used. Shader defines will always be included in preprocessed source debug dumps, but no longer explicitly added to the source by preprocessing.")
-	EDumpShaderDefines : uint8
-{
-	/** Will not be dumped unless Input.DumpDebugInfoPath is set */
-	DontCare,
-	/** No defines */
-	DontIncludeDefines,
-	/** Defines will be added in the comments */
-	AlwaysIncludeDefines
-};
-
-UE_DEPRECATED(5.4, "Please use overload of PreprocessShader accepting a FShaderPreprocessOutput struct.")
-extern SHADERPREPROCESSOR_API bool PreprocessShader(
-	FString& OutPreprocessedShader,
-	FShaderCompilerOutput& ShaderOutput,
-	const FShaderCompilerInput& ShaderInput,
-	PRAGMA_DISABLE_DEPRECATION_WARNINGS		// FShaderCompilerDefinitions will be made internal in the future, marked deprecated until then
-	const FShaderCompilerDefinitions& AdditionalDefines,
-	EDumpShaderDefines DefinesPolicy = EDumpShaderDefines::DontCare
-	PRAGMA_ENABLE_DEPRECATION_WARNINGS
-	);
-
 /**
  * Preprocess a shader.
  * @param Output - Preprocess output struct. Source, directives and possibly errors will be populated.
@@ -49,22 +26,18 @@ extern SHADERPREPROCESSOR_API bool PreprocessShader(
 	FShaderPreprocessOutput& Output,
 	const FShaderCompilerInput& Input,
 	const FShaderCompilerEnvironment& MergedEnvironment,
-	PRAGMA_DISABLE_DEPRECATION_WARNINGS		// FShaderCompilerDefinitions will be made internal in the future, marked deprecated until then
-	const FShaderCompilerDefinitions& AdditionalDefines = FShaderCompilerDefinitions()
-	PRAGMA_ENABLE_DEPRECATION_WARNINGS
-);
+	const FShaderCompilerDefinitions& AdditionalDefines);
 
-UE_DEPRECATED(5.4, "EDumpShaderDefines is no longer used. Shader defines will always be included in preprocessed source debug dumps, but no longer explicitly added to the source by preprocessing.")
-inline bool PreprocessShader(
+/**
+ * Preprocess a shader.
+ * @param Output - Preprocess output struct. Source, directives and possibly errors will be populated.
+ * @param Input - The shader compiler input.
+ * @param MergedEnvironment - The result of merging the Environment and SharedEnvironment from the FShaderCompilerInput
+ * (it is assumed this overload is called outside of the worker process which merges this in-place, so this merge step must be
+ * performed by the caller)
+ * @returns true if the shader is preprocessed without error.
+ */
+extern SHADERPREPROCESSOR_API bool PreprocessShader(
 	FShaderPreprocessOutput& Output,
 	const FShaderCompilerInput& Input,
-	const FShaderCompilerEnvironment& MergedEnvironment,
-	PRAGMA_DISABLE_DEPRECATION_WARNINGS		// FShaderCompilerDefinitions will be made internal in the future, marked deprecated until then
-	const FShaderCompilerDefinitions& AdditionalDefines,
-	EDumpShaderDefines DefinesPolicy
-	PRAGMA_ENABLE_DEPRECATION_WARNINGS
-)
-{
-	return PreprocessShader(Output, Input, MergedEnvironment, AdditionalDefines);
-}
-
+	const FShaderCompilerEnvironment& MergedEnvironment);
diff --git a/Engine/Source/Developer/TargetPlatform/Public/Interfaces/IShaderFormat.h b/Engine/Source/Developer/TargetPlatform/Public/Interfaces/IShaderFormat.h
--- a/Engine/Source/Developer/TargetPlatform/Public/Interfaces/IShaderFormat.h
+++ b/Engine/Source/Developer/TargetPlatform/Public/Interfaces/IShaderFormat.h
@@ -20,9 +20,6 @@ class IShaderFormat
 {
 public:
 
-	UE_DEPRECATED(5.4, "Monolithic CompileShader is deprecated. Please implement separate compilation and preprocess via PreprocessShader/CompilePreprocessedShader")
-	virtual void CompileShader(FName Format, const FShaderCompilerInput& Input, FShaderCompilerOutput& Output, const FString& WorkingDirectory) const {};
-
 	/**
 	 * Gets the current version of the specified shader format.
 	 *
@@ -197,9 +194,6 @@ public:
 		const FShaderCompilerOutput& Output,
 		const FShaderCompilerOutput& SecondaryOutput) const {};
 
-	UE_DEPRECATED(5.4, "SupportsIndependentPreprocessing is no longer called now that all shader backends have been migrated.")
-	virtual bool SupportsIndependentPreprocessing() const { return false; }
-
 public:
 	/** Virtual destructor. */
 	virtual ~IShaderFormat() { }
diff --git a/Engine/Source/Developer/VulkanShaderFormat/Private/SpirVShaderCompiler.inl b/Engine/Source/Developer/VulkanShaderFormat/Private/SpirVShaderCompiler.inl
--- a/Engine/Source/Developer/VulkanShaderFormat/Private/SpirVShaderCompiler.inl
+++ b/Engine/Source/Developer/VulkanShaderFormat/Private/SpirVShaderCompiler.inl
@@ -1278,9 +1278,7 @@ static bool CompileWithShaderConductor(
 	CrossCompiler::FShaderConductorContext CompilerContext;
 
 	// Inject additional macro definitions to circumvent missing features: external textures
-	PRAGMA_DISABLE_DEPRECATION_WARNINGS		// FShaderCompilerDefinitions will be made internal in the future, marked deprecated until then
 	FShaderCompilerDefinitions AdditionalDefines;
-	PRAGMA_ENABLE_DEPRECATION_WARNINGS
 
 	TArray<FString> ExtraDxcArgs;
 	if (InternalState.IsSM6())
diff --git a/Engine/Source/Developer/Windows/ShaderFormatD3D/Private/D3DShaderCompiler.cpp b/Engine/Source/Developer/Windows/ShaderFormatD3D/Private/D3DShaderCompiler.cpp
--- a/Engine/Source/Developer/Windows/ShaderFormatD3D/Private/D3DShaderCompiler.cpp
+++ b/Engine/Source/Developer/Windows/ShaderFormatD3D/Private/D3DShaderCompiler.cpp
@@ -837,7 +837,6 @@ static bool CompileAndProcessD3DShaderFXCExt(
 			TargetDesc.Language = CrossCompiler::EShaderConductorLanguage::Hlsl;
 			TargetDesc.Version = 50;
 
-			PRAGMA_DISABLE_DEPRECATION_WARNINGS		// FShaderCompilerDefinitions will be made internal in the future, marked deprecated until then
 			TargetDesc.CompileFlags->SetDefine(TEXT("implicit_resource_binding"), 1);
 			TargetDesc.CompileFlags->SetDefine(TEXT("reconstruct_global_uniforms"), 1);
 			TargetDesc.CompileFlags->SetDefine(TEXT("reconstruct_cbuffer_names"), 1);
@@ -845,7 +844,6 @@ static bool CompileAndProcessD3DShaderFXCExt(
 			TargetDesc.CompileFlags->SetDefine(TEXT("force_zero_initialized_variables"), 1);
 			TargetDesc.CompileFlags->SetDefine(TEXT("relax_nan_checks"), 1);
 			TargetDesc.CompileFlags->SetDefine(TEXT("preserve_structured_buffers"), 1);
-			PRAGMA_ENABLE_DEPRECATION_WARNINGS
 
 			// Patch SPIR-V for workarounds to prevent potential additional FXC failures
 			PatchSpirvForPrecompilation(Spirv);
diff --git a/Engine/Source/Runtime/Engine/Private/ShaderCompiler/ShaderGenerationUtil.cpp b/Engine/Source/Runtime/Engine/Private/ShaderCompiler/ShaderGenerationUtil.cpp
--- a/Engine/Source/Runtime/Engine/Private/ShaderCompiler/ShaderGenerationUtil.cpp
+++ b/Engine/Source/Runtime/Engine/Private/ShaderCompiler/ShaderGenerationUtil.cpp
@@ -39,8 +39,6 @@ bool NeedsVelocityDepth(EShaderPlatform TargetPlatform)
 
 #define SET_COMPILE_BOOL_IF_TRUE(X) { if (DerivedDefines.X) { OutEnvironment.SetDefine(TEXT(#X), 1); } }
 
-// disable deprecation warnings due to FShaderCompilerDefinitions moving to internal. This can be removed when the move is complete.
-PRAGMA_DISABLE_DEPRECATION_WARNINGS
 #define FETCH_COMPILE_BOOL(X) {																\
 	static FShaderCompilerDefineNameCache Cache_##X(TEXT(#X));								\
 	SrcDefines.X = Container.GetIntegerValue(Cache_##X, SrcDefines.X) != 0 ? 1 : 0; }
@@ -248,8 +246,6 @@ void ApplyFetchEnvironmentInternal(FShaderCompilerDefines& SrcDefines, const Con
 
 	FETCH_COMPILE_BOOL(PLATFORM_SUPPORTS_DEVELOPMENT_SHADERS);
 }
-PRAGMA_ENABLE_DEPRECATION_WARNINGS
-
 
 void FShaderCompileUtilities::ApplyFetchEnvironment(FShaderGlobalDefines& SrcDefines, const FShaderCompilerEnvironment& Environment, const EShaderPlatform Platform)
 {
@@ -303,8 +299,6 @@ struct FShaderInitialDefinesInitializer
 		ApplyFetchEnvironmentInternal(LightmapDefines, GatherNames);
 		ApplyFetchEnvironmentInternal(CompilerDefines, GatherNames);
 
-		PRAGMA_DISABLE_DEPRECATION_WARNINGS		// FShaderCompilerDefinitions will be made internal in the future, marked deprecated until then
-
 		FShaderCompilerDefinitions InitialDefines;
 		for (FName DefineKey : GatherNames.Names)
 		{
@@ -312,8 +306,6 @@ struct FShaderInitialDefinesInitializer
 			InitialDefines.FindOrAddMapIndex(DefineKey);
 		}
 		FShaderCompilerDefinitions::InitializeInitialDefines(InitialDefines);
-
-		PRAGMA_ENABLE_DEPRECATION_WARNINGS
 	}
 };
 static FShaderInitialDefinesInitializer GInitialDefinesInitializer;
diff --git a/Engine/Source/Runtime/Engine/Public/ShaderCompiler.h b/Engine/Source/Runtime/Engine/Public/ShaderCompiler.h
--- a/Engine/Source/Runtime/Engine/Public/ShaderCompiler.h
+++ b/Engine/Source/Runtime/Engine/Public/ShaderCompiler.h
@@ -980,9 +980,6 @@ public:
 	ENGINE_API int32 GetNumPendingJobs() const;
 	ENGINE_API int32 GetNumOutstandingJobs() const;
 
-	UE_DEPRECATED(5.4, "FShaderCompilingManager::GatherAnalytics is deprecated; use GatherShaderAnalytics function from ShaderAnalytics.h instead")
-	inline void GatherAnalytics(TArray<FAnalyticsEventAttribute>& Attributes) const {}
-
 	/** 
 	 * Returns whether to display a notification that shader compiling is happening in the background. 
 	 * Note: This is dependent on NumOutstandingJobs which is updated from another thread, so the results are non-deterministic.
@@ -1155,9 +1152,6 @@ public:
 	 */
 	ENGINE_API void PrintStats();
 
-	UE_DEPRECATED(5.4, "PrintStats no longer accepts a 'force' boolean since printing on an interval has become the responsibility of FShaderStatsCollector. Call the no-argument version instead.")
-	inline void PrintStats(bool bForceLogIgnoringTimeInterval) { PrintStats(); }
-
 	/** Retrieve compiler statistics for all compilation done in this process.
 	  *	Note that this will not include stats for any compilation that occurs in worker processes in a multiprocess cook. 
 	  */
@@ -1442,8 +1436,6 @@ extern ENGINE_API const FString& GetMaterialShaderMapDDCKey();
 extern ENGINE_API const FGuid& GetMaterialShaderMapDDCGuid();
 
 extern ENGINE_API bool ShouldDumpShaderDDCKeys();
-UE_DEPRECATED(5.4, "DumpShaderDDCKeyToFile now takes DebugGroupName as parameter (these files now go into the ShaderDebugInfo folder alongside other debug artifacts).")
-extern ENGINE_API void DumpShaderDDCKeyToFile(const EShaderPlatform InPlatform, bool bWithEditor, const FString& FileName, const FString& DDCKey);
 extern ENGINE_API void DumpShaderDDCKeyToFile(const EShaderPlatform InPlatform, bool bWithEditor, const TCHAR* DebugGroupName, const FString& DDCKey);
 
 /**
@@ -1453,10 +1445,4 @@ extern ENGINE_API void DumpShaderDDCKeyToFile(const EShaderPlatform InPlatform,
 * @param MaterialsToLoad				The materials contained in the MeshMaterialMaps
 * @param GlobalShaderMap				Byte array that contains the serialized global shadermap from across the network.
 **/
-extern ENGINE_API void ProcessCookOnTheFlyShaders(bool bReloadGlobalShaders, const TArray<uint8>& MeshMaterialMaps, const TArray<FString>& MaterialsToLoad, const TArray<uint8>& GlobalShaderMap);
-
-#if UE_ENABLE_INCLUDE_ORDER_DEPRECATED_IN_5_4
-#include "CoreMinimal.h"
-#include "AsyncCompilationHelpers.h"
-#include "AssetCompilingManager.h"
-#endif
+extern ENGINE_API void ProcessCookOnTheFlyShaders(bool bReloadGlobalShaders, const TArray<uint8>& MeshMaterialMaps, const TArray<FString>& MaterialsToLoad, const TArray<uint8>& GlobalShaderMap);
\ No newline at end of file
diff --git a/Engine/Source/Runtime/RenderCore/Internal/ShaderCompilerDefinitions.h b/Engine/Source/Runtime/RenderCore/Internal/ShaderCompilerDefinitions.h
--- a/Engine/Source/Runtime/RenderCore/Internal/ShaderCompilerDefinitions.h
+++ b/Engine/Source/Runtime/RenderCore/Internal/ShaderCompilerDefinitions.h
@@ -2,13 +2,10 @@
 
 #pragma once
 
-// Currently, "ShaderCore.h" includes this header, so this is a circular include, but we will remove this header from "ShaderCore.h" in 5.5,
-// as part of hiding FShaderCompilerDefinitions from public view, and will need to add this include back in.
-// #include "ShaderCore.h"
-
 #include "HAL/Platform.h"
 #include "Misc/CoreMiscDefines.h"
 #include "UObject/NameTypes.h"
+#include "ShaderCore.h"
 #include "ShaderParameterMetadata.h"
 
 #define SHADER_COMPILER_FLOAT32_FORMAT_STRING TEXT("%#.9gf")
@@ -22,10 +19,6 @@ enum class EShaderCompilerDefineVariant : uint8
 	String
 };
 
-class UE_DEPRECATED(5.4, "FShaderCompilerDefinitions will be made internal in the future, use FShaderCompilerEnvironment instead.") FShaderCompilerDefinitions;
-
-PRAGMA_DISABLE_DEPRECATION_WARNINGS
-
 /** Container for shader compiler definitions. */
 class FShaderCompilerDefinitions
 {
@@ -480,5 +473,3 @@ private:
 
 	friend struct FShaderInitialDefinesInitializer;
 };
-
-PRAGMA_ENABLE_DEPRECATION_WARNINGS
diff --git a/Engine/Source/Runtime/RenderCore/Private/ShaderCore.cpp b/Engine/Source/Runtime/RenderCore/Private/ShaderCore.cpp
--- a/Engine/Source/Runtime/RenderCore/Private/ShaderCore.cpp
+++ b/Engine/Source/Runtime/RenderCore/Private/ShaderCore.cpp
@@ -171,8 +171,6 @@ DEFINE_STAT(STAT_Shaders_ShaderPreloadMemory);
 
 CSV_DEFINE_CATEGORY(Shaders, (!UE_BUILD_SHIPPING));
 
-PRAGMA_DISABLE_DEPRECATION_WARNINGS		// FShaderCompilerDefinitions will be made internal in the future, marked deprecated until then
-
 /**
  * Singleton initial set of defines added when constructing a defines structure with bIncludeInitialDefines==true.  The
  * advantage of using preset defines is that the index of the initial define can be cached in the FShaderCompilerDefineNameCache
@@ -205,9 +203,6 @@ void FShaderCompilerDefinitions::InitializeInitialDefines(const FShaderCompilerD
 	GInitialDefines->InitialDefineCount = InDefines.Pairs.Num();
 }
 
-PRAGMA_ENABLE_DEPRECATION_WARNINGS
-
-
 // Apply lock striping as we're mostly reader lock bound.
 constexpr int32 GSHADERFILECACHE_BUCKETS = 31; /* prime number for best distribution using modulo */
 
@@ -937,9 +932,7 @@ FShaderCompilerEnvironment::FShaderCompilerEnvironment()
 	// Enable initial defines in FShaderCompilerEnvironment to improve performance (helpful here, but not for defines declared in various shader compiler backends).
 	const bool bIncludeInitialDefines = true;
 
-	PRAGMA_DISABLE_DEPRECATION_WARNINGS		// FShaderCompilerDefinitions will be made internal in the future, marked deprecated until then
 	Definitions = MakePimpl<FShaderCompilerDefinitions, EPimplPtrMode::DeepCopy>(bIncludeInitialDefines);
-	PRAGMA_ENABLE_DEPRECATION_WARNINGS
 
 	// Presize to reduce re-hashing while building shader jobs
 	IncludeVirtualPathToContentsMap.Empty(15);
@@ -949,14 +942,6 @@ FShaderCompilerEnvironment::FShaderCompilerEnvironment(FMemoryHasherBlake3& Hash
 {
 }
 
-/** Initialization constructor. */
-PRAGMA_DISABLE_DEPRECATION_WARNINGS		// FShaderCompilerDefinitions will be made internal in the future, marked deprecated until then
-FShaderCompilerEnvironment::FShaderCompilerEnvironment(const FShaderCompilerDefinitions& InDefinitions)
-{
-	Definitions = MakePimpl<FShaderCompilerDefinitions, EPimplPtrMode::DeepCopy>(InDefinitions);
-}
-PRAGMA_ENABLE_DEPRECATION_WARNINGS
-
 void FShaderCompilerEnvironment::Merge(const FShaderCompilerEnvironment& Other)
 {
 	// Merge the include maps
@@ -991,10 +976,8 @@ void FShaderCompilerEnvironment::Merge(const FShaderCompilerEnvironment& Other)
 			}
 		}
 	}
-	PRAGMA_DISABLE_DEPRECATION_WARNINGS	// FShaderCompilerDefinitions will be made internal in the future, marked deprecated until then
 	checkf(Definitions.IsValid(), TEXT("Merge is not supported on FShaderCompilerEnvironment in hashing mode"));
 	Definitions->Merge(*Other.Definitions);
-	PRAGMA_ENABLE_DEPRECATION_WARNINGS
 	CompileArgs.Append(Other.CompileArgs);
 	RenderTargetOutputFormatsMap.Append(Other.RenderTargetOutputFormatsMap);
 	FullPrecisionInPS |= Other.FullPrecisionInPS;
@@ -1004,10 +987,8 @@ FString FShaderCompilerEnvironment::GetDefinitionsAsCommentedCode() const
 {
 	checkf(Definitions.IsValid(), TEXT("GetDefinitionsAsCommentedCode is not supported on FShaderCompilerEnvironment in hashing mode"));
 	TArray<FString> DefinesLines;
-	PRAGMA_DISABLE_DEPRECATION_WARNINGS // FShaderCompilerDefinitions will be made internal in the future, marked deprecated until then
 	DefinesLines.Reserve(Definitions->Num());
 	for (FShaderCompilerDefinitions::FConstIterator DefineIt(*Definitions); DefineIt; ++DefineIt)
-	PRAGMA_ENABLE_DEPRECATION_WARNINGS
 	{
 		DefinesLines.Add(FString::Printf(TEXT("// #define %s %s\n"), DefineIt.Key(), DefineIt.Value()));
 	}
@@ -1022,8 +1003,6 @@ FString FShaderCompilerEnvironment::GetDefinitionsAsCommentedCode() const
 	return MakeInjectedShaderCodeBlock(TEXT("DumpShaderDefinesAsCommentedCode"), Defines);
 }
 
-PRAGMA_DISABLE_DEPRECATION_WARNINGS		// FShaderCompilerDefinitions will be made internal in the future, marked deprecated until then
-
 // Pass through functions to definitions
 void FShaderCompilerEnvironment::SetDefine(const TCHAR* Name, const TCHAR* Value)
 {
@@ -1267,8 +1246,6 @@ bool FShaderCompilerEnvironment::ContainsDefinition(FName Name) const
 	return false;
 }
 
-PRAGMA_ENABLE_DEPRECATION_WARNINGS
-
 /** This "core" serialization is also used for the hashing the compiler job (where files are handled differently). Should stay in sync with the ShaderCompileWorker. */
 void FShaderCompilerEnvironment::SerializeEverythingButFiles(FArchive& Ar)
 {
diff --git a/Engine/Source/Runtime/RenderCore/Public/Shader.h b/Engine/Source/Runtime/RenderCore/Public/Shader.h
--- a/Engine/Source/Runtime/RenderCore/Public/Shader.h
+++ b/Engine/Source/Runtime/RenderCore/Public/Shader.h
@@ -611,82 +611,6 @@ public:
 
 /** Encapsulates information about a shader's serialization behavior, used to detect when C++ serialization changes to auto-recompile. */
 
-class UE_DEPRECATED(5.4, "FSerializationHistory is no longer used and will be removed") FSerializationHistory 
-{
-public: 
-
-	/** Token stream stored as uint32's.  Each token is 4 bits, with a 0 meaning there's an associated 32 bit value in FullLengths. */
-	TArray<uint32> TokenBits;
-
-	/** Number of tokens in TokenBits. */
-	int32 NumTokens;
-
-	/** Full size length entries. One of these are used for every token with a value of 0. */
-	TArray<uint32> FullLengths;
-
-	FSerializationHistory() :
-		NumTokens(0)
-	{}
-
-	FORCEINLINE void AddValue(uint32 InValue)
-	{
-		const int32 UIntIndex = NumTokens / 8; 
-
-		if (UIntIndex >= TokenBits.Num())
-		{
-			// Add another uint32 if needed
-			TokenBits.AddZeroed();
-		}
-
-		uint8 Token;
-
-		// Anything that does not fit in 4 bits needs to go into FullLengths, with a special token value of 0
-		// InValue == 0 also should go into FullLengths, because its Token value is also 0
-		if (InValue > 7 || InValue == 0)
-		{
-			Token = 0;
-			FullLengths.Add(InValue);
-		}
-		else
-		{
-			Token = (uint8)InValue;
-		}
-
-		const uint32 Shift = (NumTokens % 8) * 4;
-		// Add the new token bits into the existing uint32
-		TokenBits[UIntIndex] = TokenBits[UIntIndex] | (Token << Shift);
-		NumTokens++;
-	}
-
-	FORCEINLINE uint8 GetToken(int32 Index) const
-	{
-		checkSlow(Index < NumTokens);
-		const int32 UIntIndex = Index / 8; 
-		checkSlow(UIntIndex < TokenBits.Num());
-		const uint32 Shift = (Index % 8) * 4;
-		const uint8 Token = (TokenBits[UIntIndex] >> Shift) & 0xF;
-		return Token;
-	}
-
-	void AppendKeyString(FString& KeyString) const
-	{
-		KeyString += FString::FromInt(NumTokens);
-		KeyString += BytesToHex((uint8*)TokenBits.GetData(), TokenBits.Num() * TokenBits.GetTypeSize());
-		KeyString += BytesToHex((uint8*)FullLengths.GetData(), FullLengths.Num() * FullLengths.GetTypeSize());
-	}
-
-	inline bool operator==(const FSerializationHistory& Other) const
-	{
-		return TokenBits == Other.TokenBits && NumTokens == Other.NumTokens && FullLengths == Other.FullLengths;
-	}
-
-	friend FArchive& operator<<(FArchive& Ar,class FSerializationHistory& Ref)
-	{
-		Ar << Ref.TokenBits << Ref.NumTokens << Ref.FullLengths;
-		return Ar;
-	}
-};
-
 class FShaderKey
 {
 public:
@@ -2563,51 +2487,6 @@ public:
 	{}
 };
 
-/** Tracks state when traversing a FSerializationHistory. */
-class UE_DEPRECATED(5.4, "FSerializationHistoryTraversalState is no longer used and will be removed") FSerializationHistoryTraversalState
-{
-public:
-	PRAGMA_DISABLE_DEPRECATION_WARNINGS
-	FSerializationHistoryTraversalState(const FSerializationHistory& InHistory)
-	PRAGMA_ENABLE_DEPRECATION_WARNINGS
-	{}
-
-	/** Gets the length value from NextTokenIndex + Offset into history. */
-	uint32 GetValue(int32 Offset)
-	{
-		return 0;
-	}
-
-	FORCEINLINE void StepForward()
-	{
-	}
-
-	void StepBackward()
-	{
-	}
-};
-
-/** Archive used when saving shaders, which generates data used to detect serialization mismatches on load. */
-class UE_DEPRECATED(5.4, "FShaderSaveArchive is no longer used and will be removed") FShaderSaveArchive final : public FArchiveProxy
-{
-public:
-
-	PRAGMA_DISABLE_DEPRECATION_WARNINGS
-	FShaderSaveArchive(FArchive& Archive, FSerializationHistory& InHistory) : 
-		FArchiveProxy(Archive)
-	PRAGMA_ENABLE_DEPRECATION_WARNINGS
-	{
-	}
-
-	virtual ~FShaderSaveArchive()
-	{
-	}
-
-	virtual void Serialize( void* V, int64 Length )
-	{
-	}
-};
-
 /**
  * Dumps shader stats to the log. Will also print some shader pipeline information.
  * @param Platform  - Platform to dump shader info for, use SP_NumPlatforms for all
diff --git a/Engine/Source/Runtime/RenderCore/Public/ShaderCompilerCore.h b/Engine/Source/Runtime/RenderCore/Public/ShaderCompilerCore.h
--- a/Engine/Source/Runtime/RenderCore/Public/ShaderCompilerCore.h
+++ b/Engine/Source/Runtime/RenderCore/Public/ShaderCompilerCore.h
@@ -19,9 +19,6 @@
 #include "Misc/CoreStats.h"
 #include "ShaderCore.h"
 #include "ShaderParameterMetadata.h"
-#if UE_ENABLE_INCLUDE_ORDER_DEPRECATED_IN_5_4
-#include "ShaderParameterParser.h"
-#endif
 
 class Error;
 class IShaderFormat;
@@ -266,17 +263,9 @@ struct FShaderCompilerInput
 
 	uint32 SupportedHardwareMask = 0;
 
-	// Skips the preprocessor and instead loads the usf file directly
-	UE_DEPRECATED(5.4, "bSkipPreprocessedCache member is deprecated; set EShaderDebugInfoFlags::CompileFromDebugUSF on DebugInfoFlags instead.")
-	bool bSkipPreprocessedCache = false;
-
 	// Indicates which additional debug outputs should be written for this compile job.
 	EShaderDebugInfoFlags DebugInfoFlags = EShaderDebugInfoFlags::Default;
 
-	UE_DEPRECATED(5.4, "bIndependentPreprocessed member no longer used now that all backends have been migrated to the new IShaderFormat API")
-	// True if the backend for this job implements the independent preprocessing API.
-	bool bIndependentPreprocessed = true;
-
 	UE_DEPRECATED(5.5, "bCachePreprocessed member no longer used; preprocessed job cache is now the only path for in-memory caching of shader jobs")
 	// True if the cache key for this job should be based on preprocessed source. If so,
 	// preprocessing will be executed in the cook process independent of compilation (and
@@ -328,9 +317,8 @@ struct FShaderCompilerInput
 	FOodleDataCompression::ECompressionLevel OodleLevel;
 
 	PRAGMA_DISABLE_DEPRECATION_WARNINGS
-	// Explicitly-defaulted ctors & assignment operators are needed temporarily due to 
-	// deprecation of bSkipPreprocessedCache/bIndependentPreprocessed fields. These can be removed once 
-	// the deprecation window for said fields ends.
+	// Explicitly-defaulted ctors & assignment operators are needed temporarily due to deprecation
+	// of bCachePreprocessed field. Can be removed once the deprecation window for said field ends.
 	FShaderCompilerInput() = default;
 	FShaderCompilerInput(FShaderCompilerInput&&) = default;
 	FShaderCompilerInput(const FShaderCompilerInput&) = default;
@@ -589,9 +577,8 @@ struct FShaderCompilerOutput
 	TArray<FString> UsedAttributes;
 
 	PRAGMA_DISABLE_DEPRECATION_WARNINGS
-	// Explicitly-defaulted copy/move ctors & assignment operators are needed temporarily due to 
-	// deprecation of OptionalFinalShaderSource and PragmaDirectives fields. These can be removed 
-	// once the deprecation window for said fields ends.
+	// Explicitly-defaulted copy/move ctors & assignment operators are needed temporarily due to deprecation
+	// of PragmaDirectives field. These can be removed once the deprecation window for said field ends.
 	FShaderCompilerOutput(FShaderCompilerOutput&&) = default;
 	FShaderCompilerOutput(const FShaderCompilerOutput&) = default;
 	FShaderCompilerOutput& operator=(FShaderCompilerOutput&&) = default;
@@ -605,9 +592,6 @@ struct FShaderCompilerOutput
 
 	TArray<FShaderDiagnosticData> ShaderDiagnosticDatas;
 
-	UE_DEPRECATED(5.4, "OptionalFinalShaderSource is no longer used; set ModifiedShaderSource instead if the shader backend makes additional source code manipulations.")
-	FString OptionalFinalShaderSource;
-
 	/** Use this field to store the shader source code if it's modified as part of the shader format's compilation process. This should only be set when 
 	 * additional manipulation of source code  is required that is not part of the implementation of PreprocessShader. This version of the source, if set, 
 	 * will be what is written as part of the debug dumps of preprocessed source, as well as used for upstream code which explicitly requests the final
diff --git a/Engine/Source/Runtime/RenderCore/Public/ShaderCompilerJobTypes.h b/Engine/Source/Runtime/RenderCore/Public/ShaderCompilerJobTypes.h
--- a/Engine/Source/Runtime/RenderCore/Public/ShaderCompilerJobTypes.h
+++ b/Engine/Source/Runtime/RenderCore/Public/ShaderCompilerJobTypes.h
@@ -279,13 +279,6 @@ public:
 	// Serializes only the subset of data written by ShaderCompiler and read from SCW when using worker processes.
 	RENDERCORE_API void SerializeWorkerInput(FArchive& Ar);
 
-	UE_DEPRECATED(5.4, "GetFinalSource is deprecated, GetFinalSourceView returns an FStringView instead.")
-	inline const FString& GetFinalSource() const
-	{
-		static FString Empty;
-		return Empty;
-	}
-
 	RENDERCORE_API FStringView GetFinalSourceView() const;
 
 	FShaderCompileJob() : FShaderCommonCompileJob(Type, 0u, 0u, EShaderCompileJobPriority::Num)
diff --git a/Engine/Source/Runtime/RenderCore/Public/ShaderCore.h b/Engine/Source/Runtime/RenderCore/Public/ShaderCore.h
--- a/Engine/Source/Runtime/RenderCore/Public/ShaderCore.h
+++ b/Engine/Source/Runtime/RenderCore/Public/ShaderCore.h
@@ -42,9 +42,6 @@
 #include "UObject/UnrealNames.h"
 #include "UniformBuffer.h"
 
-// Temporarily included here until we can fully deprecate access to FShaderCompilerDefinitions in a future version
-#include "ShaderCompilerDefinitions.h"
-
 class Error;
 class FMemoryImageWriter;
 class FMemoryUnfreezeContent;
@@ -504,6 +501,30 @@ struct FShaderResourceTableMap
 	RENDERCORE_API void FixupOnLoad(const TMap<FString, FUniformBufferEntry>& UniformBufferMap);
 };
 
+class FShaderCompilerDefineNameCache
+{
+public:
+	FShaderCompilerDefineNameCache(const TCHAR* InName)
+		: Name(InName), MapIndex(INDEX_NONE)
+	{}
+
+	FString ToString() const 
+	{
+		return Name.ToString();
+	}
+
+	operator FName() const
+	{
+		return Name;
+	}
+
+private:
+	FName Name;
+	int32 MapIndex;
+
+	friend class FShaderCompilerDefinitions;
+};
+
 /** The environment used to compile a shader. */
 struct FShaderCompilerEnvironment
 {
@@ -513,9 +534,6 @@ struct FShaderCompilerEnvironment
 
 	TMap<FString, FThreadSafeSharedAnsiStringPtr> IncludeVirtualPathToSharedContentsMap;
 
-	UE_DEPRECATED(5.4, "IncludeVirtualPathToExternalContentsMap has been replaced with IncludeVirtualPathToSharedContentsMap (type change from FString to ANSI string).")
-	TMap<FString, FThreadSafeSharedStringPtr> IncludeVirtualPathToExternalContentsMap;
-
 	FShaderCompilerFlags CompilerFlags;
 	TMap<uint32,uint8> RenderTargetOutputFormatsMap;
 	FShaderResourceTableMap ResourceTableMap;
@@ -538,27 +556,6 @@ struct FShaderCompilerEnvironment
 	/** Constructor used when enviroment is constructed temporarily purely for the purpose of hashing for inclusion in DDC keys. */
 	RENDERCORE_API FShaderCompilerEnvironment(FMemoryHasherBlake3& Hasher);
 
-	/** Initialization constructor. */
-	PRAGMA_DISABLE_DEPRECATION_WARNINGS
-	UE_DEPRECATED(5.4, "FShaderCompilerDefinitions is being made private in the future, do not use this constructor.")
-	RENDERCORE_API explicit FShaderCompilerEnvironment(const FShaderCompilerDefinitions& InDefinitions);
-	PRAGMA_ENABLE_DEPRECATION_WARNINGS
-
-	// Used as a baseclass, make sure we're not incorrectly destroyed through a baseclass pointer
-	// This will be expensive to destroy anyway, additional vcall overhead should be small
-	PRAGMA_DISABLE_DEPRECATION_WARNINGS
-	// Temporarily disable deprecation warnings for this default destructor triggering due to the 
-	// deprecated RemoteServerData field. When the field is removed warning disable should be as 
-	// well, but the defaulted destructor should remain due to the above comment.
-	virtual ~FShaderCompilerEnvironment() = default;
-
-	// Explicitly default assignment operator and copy constructor operator with warnings disabled
-	// to avoid warnings in implicitly-generated functions due to deprecation of IncludeVirtualPathToExternalContentsMap. 
-	// These can be removed entirely (revert to implicitly-generated) when the field itself is.
-	FShaderCompilerEnvironment(const FShaderCompilerEnvironment&) = default;
-	FShaderCompilerEnvironment& operator=(const FShaderCompilerEnvironment&) = default;
-	PRAGMA_ENABLE_DEPRECATION_WARNINGS
-
 	/**
 	 * Works for TCHAR
 	 * e.g. SetDefine(TEXT("NAME"), TEXT("Test"));
@@ -701,9 +698,7 @@ private:
 	friend class FShaderCompileUtilities;
 	friend class FShaderPreprocessorUtilities;
 
-	PRAGMA_DISABLE_DEPRECATION_WARNINGS		// FShaderCompilerDefinitions will be made internal in the future, marked deprecated until then
 	TPimplPtr<FShaderCompilerDefinitions, EPimplPtrMode::DeepCopy> Definitions;
-	PRAGMA_ENABLE_DEPRECATION_WARNINGS
 
 	FMemoryHasherBlake3* Hasher = nullptr;
 
diff --git a/Engine/Source/Runtime/RenderCore/Public/ShaderParameterMetadata.h b/Engine/Source/Runtime/RenderCore/Public/ShaderParameterMetadata.h
--- a/Engine/Source/Runtime/RenderCore/Public/ShaderParameterMetadata.h
+++ b/Engine/Source/Runtime/RenderCore/Public/ShaderParameterMetadata.h
@@ -542,39 +542,3 @@ private:
 };
 
 ENUM_CLASS_FLAGS(FShaderParametersMetadata::EUsageFlags);
-
-
-/**
- * Utility class for caching FName and other info for a shader compiler define.  Don't use directly, use SET_SHADER_DEFINE or SET_SHADER_DEFINE_AND_COMPILE_ARGUMENT
- * macros defined in ShaderCore.h.
- *
- * Placed here temporarily to solve static analysis warning related to deprecation of public access to FShaderCompilerDefinitions.  This structure needs to be
- * declared before including "Runtime/RenderCore/Internal/ShaderCompilerDefinitions.h", which would normally be accomplished by including ShaderCore.h, but we
- * can't include that without producing a circular include, which generates a static analysis warning.  The first attempt to solve this involved placing the
- * include after this structure declaration in ShaderCore.h, but that leads to an "Include after first code block" static analysis warning.  We don't want to move
- * this structure to the internal header, so we need to place it in another header that's also included in "ShaderCore.h", to avoid producing that warning.  This
- * is the only candidate shader related source file, so it makes sense to place here for now, and can be moved back to ShaderCore.h next major version.
- */
-class FShaderCompilerDefineNameCache
-{
-public:
-	FShaderCompilerDefineNameCache(const TCHAR* InName)
-		: Name(InName), MapIndex(INDEX_NONE)
-	{}
-
-	FString ToString() const 
-	{
-		return Name.ToString();
-	}
-
-	operator FName() const
-	{
-		return Name;
-	}
-
-private:
-	FName Name;
-	int32 MapIndex;
-
-	friend class FShaderCompilerDefinitions;
-};
diff --git a/Engine/Source/Runtime/RenderCore/Public/ShaderPreprocessTypes.h b/Engine/Source/Runtime/RenderCore/Public/ShaderPreprocessTypes.h
--- a/Engine/Source/Runtime/RenderCore/Public/ShaderPreprocessTypes.h
+++ b/Engine/Source/Runtime/RenderCore/Public/ShaderPreprocessTypes.h
@@ -66,13 +66,6 @@ public:
 	{
 	}
 	
-	UE_DEPRECATED(5.4, "Use StringView accessors")
-	const FString& GetSource() const
-	{
-		static FString Empty;
-		return Empty;
-	}
-
 	FShaderSource::FViewType GetSourceView() const
 	{
 		return PreprocessedSource.GetView();
@@ -104,14 +97,6 @@ public:
 #endif
 	}
 
-	UE_DEPRECATED(5.4, "Use GetUnstrippedSourceView")
-	const FString& GetUnstrippedSource() const
-	{
-		PRAGMA_DISABLE_DEPRECATION_WARNINGS
-		return GetSource();
-		PRAGMA_ENABLE_DEPRECATION_WARNINGS
-	}
-
 	FStringView GetUnstrippedSourceView() const
 	{
 		// if the unstripped source is requested, check if the "original source" field has been populated
@@ -325,7 +310,3 @@ private:
 
 	TArray<FShaderDiagnosticData> ShaderDiagnosticDatas;
 };
-
-#if UE_ENABLE_INCLUDE_ORDER_DEPRECATED_IN_5_4
-#include "ShaderParameterParser.h"
-#endif
