Added missing entry for Metal 3.1 to GetMetalFamilyVersionString

diff --git a/Engine/Source/Developer/ShaderCompilerCommon/Private/ShaderConductorContext.cpp b/Engine/Source/Developer/ShaderCompilerCommon/Private/ShaderConductorContext.cpp
--- a/Engine/Source/Developer/ShaderCompilerCommon/Private/ShaderConductorContext.cpp
+++ b/Engine/Source/Developer/ShaderCompilerCommon/Private/ShaderConductorContext.cpp
@@ -195,6 +195,7 @@ namespace CrossCompiler
 	{
 		switch (Version)
 		{
+		case 30100: return "30100";
         case 30000: return "30000";
 		case 20400: return "20400";
 		case 20300: return "20300";
