[Shaders] Initialize FHlslScanner with input filename.

#rnx

diff --git a/Engine/Source/Developer/ShaderCompilerCommon/Private/HlslLexer.cpp b/Engine/Source/Developer/ShaderCompilerCommon/Private/HlslLexer.cpp
--- a/Engine/Source/Developer/ShaderCompilerCommon/Private/HlslLexer.cpp
+++ b/Engine/Source/Developer/ShaderCompilerCommon/Private/HlslLexer.cpp
@@ -1131,7 +1131,7 @@ namespace CrossCompiler
 		// Simple heuristic to avoid reallocating
 		Tokens.Reserve(String.Len() / 8);
 
-		FTokenizer Tokenizer(String);
+		FTokenizer Tokenizer(String, Filename);
 		while (Tokenizer.HasCharsAvailable())
 		{
 			auto* Sanity = Tokenizer.Current;
