* Added NoExportTypes.cpp and inlined generated file with needed includes in front (It ends up being the same as NoExportTypes generating its own cpp which is compiled)
* Fixed up "..\..\.." paths since they caused vfsoverlay path to produce bad dependency files (clang bug?)

diff --git a/Engine/Plugins/Media/WebMMedia/Source/WebMMediaFactory/Private/WebMMediaFactoryModule.cpp b/Engine/Plugins/Media/WebMMedia/Source/WebMMediaFactory/Private/WebMMediaFactoryModule.cpp
--- a/Engine/Plugins/Media/WebMMedia/Source/WebMMediaFactory/Private/WebMMediaFactoryModule.cpp
+++ b/Engine/Plugins/Media/WebMMedia/Source/WebMMediaFactory/Private/WebMMediaFactoryModule.cpp
@@ -15,7 +15,7 @@
 #include "UObject/NameTypes.h"
 #include "HAL/PlatformProperties.h"
 
-#include "../../WebMMedia/Public/IWebMMediaModule.h"
+#include "WebMMedia/Public/IWebMMediaModule.h"
 
 
 DEFINE_LOG_CATEGORY(LogWebMMediaFactory);
diff --git a/Engine/Source/Developer/ShaderCompilerCommon/Public/CrossCompilerDefinitions.h b/Engine/Source/Developer/ShaderCompilerCommon/Public/CrossCompilerDefinitions.h
--- a/Engine/Source/Developer/ShaderCompilerCommon/Public/CrossCompilerDefinitions.h
+++ b/Engine/Source/Developer/ShaderCompilerCommon/Public/CrossCompilerDefinitions.h
@@ -5,5 +5,5 @@
 //@todo-lh: As soon as HLSLcc is fully replaced by DXC for all shader compiler backends,
 // move the content of this file into 'CrossCompilerDefinitions.h".
 // Until then, we simply include that file from the HLSLcc folder since it only contains enum declaration but no further include dependencies.
-#include "../../../ThirdParty/hlslcc/hlslcc/src/hlslcc_lib/HlslccDefinitions.h"
+#include "ThirdParty/hlslcc/hlslcc/src/hlslcc_lib/HlslccDefinitions.h"
 
diff --git a/Engine/Source/Developer/Windows/LiveCoding/LiveCoding.Build.cs b/Engine/Source/Developer/Windows/LiveCoding/LiveCoding.Build.cs
--- a/Engine/Source/Developer/Windows/LiveCoding/LiveCoding.Build.cs
+++ b/Engine/Source/Developer/Windows/LiveCoding/LiveCoding.Build.cs
@@ -40,6 +40,11 @@ public class LiveCoding : ModuleRules
         	PrivateDefinitions.Add("LC_DEBUG=0");
         }
 
+		if (Target.Platform == UnrealTargetPlatform.Win64)
+		{
+			PrivateIncludePaths.Add("Developer/Windows/LiveCodingServer/Private/External");
+		}
+
 		PublicAdditionalLibraries.Add(Path.Combine(ModuleDirectory, "Private", "External", "LC_JumpToSelf.lib"));
 	}
 }
diff --git a/Engine/Source/Developer/Windows/LiveCoding/Private/External/LC_Commands.h b/Engine/Source/Developer/Windows/LiveCoding/Private/External/LC_Commands.h
--- a/Engine/Source/Developer/Windows/LiveCoding/Private/External/LC_Commands.h
+++ b/Engine/Source/Developer/Windows/LiveCoding/Private/External/LC_Commands.h
@@ -4,7 +4,7 @@
 
 // BEGIN EPIC MOD
 #include "CoreTypes.h"
-#include "../../../LiveCodingServer/Private/External/LC_Hook.h"
+#include "LC_Hook.h"
 #include "Windows/WindowsHWrapper.h"
 // END EPIC MOD
 #include "LC_ProcessTypes.h"
diff --git a/Engine/Source/Runtime/CoreUObject/Private/UObject/NoExportTypes.cpp b/Engine/Source/Runtime/CoreUObject/Private/UObject/NoExportTypes.cpp
new file mode 100644
--- /dev/null
+++ b/Engine/Source/Runtime/CoreUObject/Private/UObject/NoExportTypes.cpp
@@ -0,0 +1,8 @@
+// Copyright Epic Games, Inc. All Rights Reserved.
+
+#include "UObject/NoExportTypes.h"
+#include "../../../ApplicationCore/Public/GenericPlatform/ICursor.h"
+#include "../../../ApplicationCore/Public/GenericPlatform/IInputInterface.h"
+#include "../../../ApplicationCore/Public/GenericPlatform/InputDeviceMappingPolicy.h"
+
+#include UE_INLINE_GENERATED_CPP_BY_NAME(NoExportTypes)
diff --git a/Engine/Source/Runtime/CoreUObject/Public/UObject/NoExportTypes.h b/Engine/Source/Runtime/CoreUObject/Public/UObject/NoExportTypes.h
--- a/Engine/Source/Runtime/CoreUObject/Public/UObject/NoExportTypes.h
+++ b/Engine/Source/Runtime/CoreUObject/Public/UObject/NoExportTypes.h
@@ -61,10 +61,6 @@
 #include "AssetRegistry/AssetBundleData.h"
 #include "AssetRegistry/AssetData.h"
 
-#include "../../../ApplicationCore/Public/GenericPlatform/ICursor.h"
-#include "../../../ApplicationCore/Public/GenericPlatform/IInputInterface.h"
-#include "../../../ApplicationCore/Public/GenericPlatform/InputDeviceMappingPolicy.h"
-
 #endif
 
 #if !CPP      //noexport class
diff --git a/Engine/Source/Runtime/Renderer/Private/StateStream/SkinnedMeshStateStreamImpl.cpp b/Engine/Source/Runtime/Renderer/Private/StateStream/SkinnedMeshStateStreamImpl.cpp
--- a/Engine/Source/Runtime/Renderer/Private/StateStream/SkinnedMeshStateStreamImpl.cpp
+++ b/Engine/Source/Runtime/Renderer/Private/StateStream/SkinnedMeshStateStreamImpl.cpp
@@ -5,9 +5,9 @@
 #include "Rendering/RenderCommandPipes.h"
 #include "ScenePrivate.h"
 #include "SkeletalMeshSceneProxy.h"
+#include "SkeletalRenderPublic.h"
 #include "SkinnedMeshSceneProxyDesc.h"
 #include "StateStreamCreator.h"
-#include "../../../Engine/Private/SkeletalRenderGPUSkin.h"
 
 ////////////////////////////////////////////////////////////////////////////////////////////////////
 
diff --git a/Engine/Source/Runtime/VulkanRHI/Private/VulkanChunkedPipelineCache.cpp b/Engine/Source/Runtime/VulkanRHI/Private/VulkanChunkedPipelineCache.cpp
--- a/Engine/Source/Runtime/VulkanRHI/Private/VulkanChunkedPipelineCache.cpp
+++ b/Engine/Source/Runtime/VulkanRHI/Private/VulkanChunkedPipelineCache.cpp
@@ -54,7 +54,7 @@ public:
 	FScopedTimeToLog(FString, class FOutputDevice* output = nullptr) {};
 };
 #endif
-#include "../../Core/Public/Containers/LockFreeList.h"
+#include "Containers/LockFreeList.h"
 
 #if PLATFORM_ANDROID && USE_ANDROID_FILE
 // TODO:
