[Shaders] Print token when HlslParser cannot match parsing rule and use MSVC diagnostics format for source location.

#rnx

diff --git a/Engine/Source/Developer/ShaderCompilerCommon/Private/HlslParser.cpp b/Engine/Source/Developer/ShaderCompilerCommon/Private/HlslParser.cpp
--- a/Engine/Source/Developer/ShaderCompilerCommon/Private/HlslParser.cpp
+++ b/Engine/Source/Developer/ShaderCompilerCommon/Private/HlslParser.cpp
@@ -2231,7 +2231,7 @@ Done:
 			break;
 		}
 
-		const auto* Peek = Parser.Scanner.GetCurrentToken();
+		const CrossCompiler::FHlslToken* Peek = Parser.Scanner.GetCurrentToken();
 		if (!Peek)
 		{
 			return ParseResultError();
@@ -2273,7 +2273,7 @@ Done:
 			return EParseResult::Ignore;
 		}
 
-		Parser.Scanner.SourceError(TEXT("Unable to match rule!"));
+		Parser.Scanner.SourceError(FString::Printf(TEXT("Unable to match rule! [Token=\"%s\"]"), *Peek->String));
 		return ParseResultError();
 	}
 
diff --git a/Engine/Source/Developer/ShaderCompilerCommon/Private/HlslUtils.h b/Engine/Source/Developer/ShaderCompilerCommon/Private/HlslUtils.h
--- a/Engine/Source/Developer/ShaderCompilerCommon/Private/HlslUtils.h
+++ b/Engine/Source/Developer/ShaderCompilerCommon/Private/HlslUtils.h
@@ -258,11 +258,11 @@ namespace CrossCompiler
 		{
 			if (SourceInfo.Filename)
 			{
-				AddMessage(true, FString::Printf(TEXT("%s(%d): (%d) %s\n"), **SourceInfo.Filename, SourceInfo.Line, SourceInfo.Column, String));
+				AddMessage(true, FString::Printf(TEXT("%s(%d,%d): %s\n"), **SourceInfo.Filename, SourceInfo.Line, SourceInfo.Column, String));
 			}
 			else
 			{
-				AddMessage(true, FString::Printf(TEXT("<unknown>(%d): (%d) %s\n"), SourceInfo.Line, SourceInfo.Column, String));
+				AddMessage(true, FString::Printf(TEXT("<unknown>(%d,%d): %s\n"), SourceInfo.Line, SourceInfo.Column, String));
 			}
 		}
 
@@ -275,11 +275,11 @@ namespace CrossCompiler
 		{
 			if (SourceInfo.Filename)
 			{
-				AddMessage(false, FString::Printf(TEXT("%s(%d): (%d) %s\n"), **SourceInfo.Filename, SourceInfo.Line, SourceInfo.Column, String));
+				AddMessage(false, FString::Printf(TEXT("%s(%d,%d): %s\n"), **SourceInfo.Filename, SourceInfo.Line, SourceInfo.Column, String));
 			}
 			else
 			{
-				AddMessage(false, FString::Printf(TEXT("<unknown>(%d): (%d) %s\n"), SourceInfo.Line, SourceInfo.Column, String));
+				AddMessage(false, FString::Printf(TEXT("<unknown>(%d,%d): %s\n"), SourceInfo.Line, SourceInfo.Column, String));
 			}
 		}
 
