[shaders] symbol export fixes
- change static boolean guarding initialization to member var so we can export symbols for multiple formats in a single session
- add thread safety in preparation for this being called from job completion in-editor
- when not exporting to a zipfile add extra logic to write out a stable copy of the PDB based on least hash value, similar to how the copy saved in shadermap editoronly data works (zip writer doesn't allow removing/replacing files, only adding, so we can't support a similar mechanism for symbol zip exports at this time, though this is a minor issue, non-determinism isn't causing any major problems here - this just makes my life easier when debugging symbol output)

diff --git a/Engine/Source/Developer/ShaderCompilerCommon/Private/ShaderSymbolExport.cpp b/Engine/Source/Developer/ShaderCompilerCommon/Private/ShaderSymbolExport.cpp
--- a/Engine/Source/Developer/ShaderCompilerCommon/Private/ShaderSymbolExport.cpp
+++ b/Engine/Source/Developer/ShaderCompilerCommon/Private/ShaderSymbolExport.cpp
@@ -9,6 +9,7 @@
 #include "GenericPlatform/GenericPlatformFile.h"
 #include "HAL/FileManager.h"
 #include "HAL/PlatformFileManager.h"
+#include "Hash/CityHash.h"
 #include "Logging/LogMacros.h"
 #include "Misc/CommandLine.h"
 #include "Misc/PathViews.h"
@@ -139,13 +140,46 @@ void FShaderSymbolExport::WriteSymbolData(const FString& Filename, const FString
 		return;
 	}
 
-	// Skip this symbol data if we've already exported it before.
+	// If writing to a zipfile (for simplicity's sake, since zipwriter doesn't allow removing/replacing files), or when exporting
+	// symbols info instead of/in addition to full symbols, we just always write the first copy of symbols encountered.
+	// Note this means the actual symbols written in these cases may be non-deterministic, but this won't cause any real problems.
+	// Otherwise, we add an additional condition based on the hash of the symbols, if the new hash is less than the one previously
+	// written we re-write (so we always export a deterministic copy of the symbols in this case).
+
 	bool bAlreadyInSet = false;
 	ExportedShaders.Add(Filename, &bAlreadyInSet);
+	bool bDeterministicSymbols = ZipWriter == nullptr && !ShouldGenerateShaderSymbolsInfo(ShaderFormat);
+
+	// calculate the hash of the new symbols only if we're not exporting deterministically
+	// (we don't need it in that case as described in the comments above)
+	uint64 NewSymbolHash = bDeterministicSymbols ? CityHash64(reinterpret_cast<const char*>(Contents.GetData()), Contents.Num()) : 0u;
 	if (bAlreadyInSet)
 	{
-		// We've already exported this shader hash
-		return;
+		// We've already exported a copy of symbols for this shader hash and are not ensuring exports are deterministic
+		if (!bDeterministicSymbols)
+		{
+			return;
+		}
+		
+		// when ensuring determinism, we further store the hash of the symbols we wrote out, and if the new hash is less than
+		// the previously written one we write it again, so we always choose the same copy of symbols to write during i.e. a cook
+		FSymbolFileInfo* PrevExportedSymbolInfo = ExportedSymbolInfo.Find(Filename);
+		
+		check(PrevExportedSymbolInfo);
+		check(NewSymbolHash != 0u);
+		if (NewSymbolHash >= PrevExportedSymbolInfo->Hash)
+		{
+			return;
+		}
+
+		// new hash is less; adjust stats to account for file we will replace below
+		TotalSymbolDataBytes -= PrevExportedSymbolInfo->Size;
+		TotalSymbolData--;
+	}
+
+	if (bDeterministicSymbols)
+	{
+		ExportedSymbolInfo.Add(Filename, { NewSymbolHash, Contents.Num() });
 	}
 
 	// Emit periodic log messages detailing the size of the shader symbols output file/directory.
diff --git a/Engine/Source/Developer/ShaderCompilerCommon/Public/ShaderSymbolExport.h b/Engine/Source/Developer/ShaderCompilerCommon/Public/ShaderSymbolExport.h
--- a/Engine/Source/Developer/ShaderCompilerCommon/Public/ShaderSymbolExport.h
+++ b/Engine/Source/Developer/ShaderCompilerCommon/Public/ShaderSymbolExport.h
@@ -8,9 +8,12 @@
 
 #include "CoreTypes.h"
 #include "Containers/ContainersFwd.h"
+#include "Containers/Map.h"
 #include "Containers/Set.h"
-#include "Templates/UniquePtr.h"
+#include "HAL/CriticalSection.h"
+#include "Misc/ScopeLock.h"
 #include "Serialization/MemoryReader.h"
+#include "Templates/UniquePtr.h"
 
 class FZipArchiveWriter;
 
@@ -38,11 +41,18 @@ private:
 
 	TUniquePtr<FZipArchiveWriter> ZipWriter;
 	TSet<FString> ExportedShaders;
+	struct FSymbolFileInfo
+	{
+		uint64 Hash;
+		int32 Size;
+	};
+	TMap<FString, FSymbolFileInfo> ExportedSymbolInfo;
 	FString ExportPath;
 	FString InfoFilePath;
-
+	FCriticalSection Cs;
 	uint64 TotalSymbolDataBytes{ 0 };
 	uint64 TotalSymbolData{ 0 };
+	bool bInitialized{ false };
 	bool bExportShaderSymbols{ false };
 
 	struct FShaderInfo
@@ -63,13 +73,13 @@ private:
 template<typename TPlatformShaderSymbolData>
 inline void FShaderSymbolExport::NotifyShaderCompiled(const TConstArrayView<uint8>& PlatformSymbolData, const FString& DebugInfo)
 {
-	static bool bFirst = true;
-	if (bFirst)
+	FScopeLock Lock(&Cs);
+	if (!bInitialized)
 	{
 		// If we get called, we know we're compiling. Do one time initialization
 		// which will create the output directory / open the open file stream.
 		Initialize();
-		bFirst = false;
+		bInitialized = true;
 	}
 
 	if (bExportShaderSymbols)
