Added logging for total shaders and platform deduplication count
Requires r.Shaders.SymbolsInfo=1

#tests local cook Lyra

diff --git a/Engine/Source/Developer/ShaderCompilerCommon/Private/ShaderSymbolExport.cpp b/Engine/Source/Developer/ShaderCompilerCommon/Private/ShaderSymbolExport.cpp
--- a/Engine/Source/Developer/ShaderCompilerCommon/Private/ShaderSymbolExport.cpp
+++ b/Engine/Source/Developer/ShaderCompilerCommon/Private/ShaderSymbolExport.cpp
@@ -156,6 +156,7 @@ void FShaderSymbolExport::WriteSymbolData(const FString& Filename, const FString
 	if (bAlreadyInSet)
 	{
 		// We've already exported a copy of symbols for this shader hash and are not ensuring exports are deterministic
+		DuplicateSymbols++;
 		if (!bDeterministicSymbols)
 		{
 			return;
@@ -276,7 +277,12 @@ void FShaderSymbolExport::NotifyShaderCompilersShutdown()
 								// if this symbol is new to the multiproc owner, store it
 								bool bAlreadyInSet = false;
 								ExportedShaders.Add(Filename, &bAlreadyInSet);
-								if (!bAlreadyInSet)
+								if (bAlreadyInSet)
+								{
+									// The multiproc owner has already seen this hash
+									DuplicateSymbols++;
+								}
+								else
 								{
 									FString DebugData = Line.Right(Line.Len() - Space - 1);
 									ShaderInfos.Add({ Filename, DebugData });
@@ -308,6 +314,9 @@ void FShaderSymbolExport::NotifyShaderCompilersShutdown()
 				Writer->Serialize(Output.GetData(), Output.Num());
 				Writer->Close();
 				UE_LOG(LogShaderSymbolExport, Display, TEXT("Wrote %d records into shader symbols info output file \"%s\"."), ShaderInfos.Num(), *InfoFilePath);
+				const uint32 TotalShaders = ShaderInfos.Num() + DuplicateSymbols;
+				UE_LOG(LogShaderSymbolExport, Display, TEXT("%d total shaders, %d shaders after platform deduplication. %d duplicates (%4.1f%%)."),
+					TotalShaders, TotalShaders-DuplicateSymbols, DuplicateSymbols, (float)DuplicateSymbols / TotalShaders * 100.0f);
 			}
 			else
 			{
diff --git a/Engine/Source/Developer/ShaderCompilerCommon/Public/ShaderSymbolExport.h b/Engine/Source/Developer/ShaderCompilerCommon/Public/ShaderSymbolExport.h
--- a/Engine/Source/Developer/ShaderCompilerCommon/Public/ShaderSymbolExport.h
+++ b/Engine/Source/Developer/ShaderCompilerCommon/Public/ShaderSymbolExport.h
@@ -62,6 +62,8 @@ private:
 	};
 	TArray<FShaderInfo> ShaderInfos;
 
+	uint32 DuplicateSymbols{ 0 };
+
 	/**
 	 * If true, the current process is the first process in a multiprocess group, or is not in a group,
 	 * and should combine artifacts produced by the other processes. Will also be false if no combination
