add ShaderCompilerCommon module to foundation tests; add TEST_IF_EQUAL macro to expose catch2's CHECKED_IF functionality for lowlevel tests, migrate shader minifier tests to catch2/LLT

diff --git a/Engine/Source/Developer/LowLevelTestsRunner/Public/TestHarness.h b/Engine/Source/Developer/LowLevelTestsRunner/Public/TestHarness.h
--- a/Engine/Source/Developer/LowLevelTestsRunner/Public/TestHarness.h
+++ b/Engine/Source/Developer/LowLevelTestsRunner/Public/TestHarness.h
@@ -164,6 +164,10 @@ struct Catch::StringMaker<TTuple<KeyType, ValueType>>
 	INFO(What);\
 	CHECK((Value)==false)
 
+#define TEST_IF_EQUAL(What, Actual, Expected)\
+	INFO(What);\
+	CHECKED_IF(Actual == Expected)
+
 #define TEST_EQUAL(What, Actual, Expected)\
 	INFO(What);\
 	CHECK((Actual) == (Expected))
diff --git a/Engine/Source/Developer/ShaderCompilerCommon/Private/ShaderMinifier.cpp b/Engine/Source/Developer/ShaderCompilerCommon/Private/ShaderMinifier.cpp
--- a/Engine/Source/Developer/ShaderCompilerCommon/Private/ShaderMinifier.cpp
+++ b/Engine/Source/Developer/ShaderCompilerCommon/Private/ShaderMinifier.cpp
@@ -5,7 +5,7 @@
 #include "HAL/PlatformTime.h"
 #include "Hash/xxhash.h"
 #include "Logging/LogMacros.h"
-#include "Misc/AutomationTest.h"
+#include "Tests/TestHarnessAdapter.h"
 #include "String/Find.h"
 #include "Algo/BinarySearch.h"
 #include "Misc/MemStack.h"
@@ -2118,9 +2118,7 @@ FMinifiedShader Minify(const FShaderSource& PreprocessedShader, const FShaderSou
 
 } // namespace UE::ShaderMinifier
 
-#if WITH_AUTOMATION_TESTS
-
-IMPLEMENT_SIMPLE_AUTOMATION_TEST(FShaderMinifierParserTest, "System.Shaders.ShaderMinifier.Parse", EAutomationTestFlags::EditorContext | EAutomationTestFlags::EngineFilter);
+#if WITH_LOW_LEVEL_TESTS
 
 namespace UE::ShaderMinifier
 {
@@ -2132,25 +2130,25 @@ static FParsedShader ParseShader(const FShaderSource& InSource)
 }
 }
 
-bool FShaderMinifierParserTest::RunTest(const FString& Parameters)
+TEST_CASE_NAMED(FShaderMinifierParserTest, "Shaders::ShaderMinifier::Parse", "[EditorContext][EngineFilter]")
 {
 	using namespace UE::ShaderMinifier;
 
 	{
 		FShaderSource S(SHADER_SOURCE_LITERAL("  \n\r\f \tHello"));
-		TestEqual(TEXT("SkipSpace"),
+		TEST_EQUAL(TEXT("SkipSpace"),
 			FString(SkipSpace(S.GetView())),
 			FString(SHADER_SOURCE_LITERAL("Hello")));
 	}
 	{
 		FShaderSource S(SHADER_SOURCE_LITERAL("Hello World"));
-		TestEqual(TEXT("SkipUntilStr (found)"),
+		TEST_EQUAL(TEXT("SkipUntilStr (found)"),
 			FString(SkipUntilStr(S.GetView(), SHADER_SOURCE_LITERAL("World"))),
 			FString(SHADER_SOURCE_LITERAL("World")));
 	}
 	{
 		FShaderSource S(SHADER_SOURCE_LITERAL("Hello World"));
-		TestEqual(TEXT("SkipUntilStr (not found)"),
+		TEST_EQUAL(TEXT("SkipUntilStr (not found)"),
 			FString(SkipUntilStr(S.GetView(), SHADER_SOURCE_LITERAL("Blah"))),
 			FString());
 	}
@@ -2158,20 +2156,18 @@ bool FShaderMinifierParserTest::RunTest(const FString& Parameters)
 	{
 		FShaderSource S(SHADER_SOURCE_LITERAL("static const struct { int Blah; } Foo = { 123; };"));
 		auto P = ParseShader(S);
-		TestEqual(TEXT("Anonymous struct variable with initializer, total chunks"), P.Chunks.Num(), 1);
-		if (P.Chunks.Num() == 1)
+		TEST_IF_EQUAL(TEXT("Anonymous struct variable with initializer, total chunks"), P.Chunks.Num(), 1)
 		{
-			TestEqual(TEXT("Anonymous struct variable with initializer, main chunk type"), P.Chunks[0].Type, ECodeChunkType::Variable);
+			TEST_EQUAL(TEXT("Anonymous struct variable with initializer, main chunk type"), P.Chunks[0].Type, ECodeChunkType::Variable);
 		}
 	}
 
 	{
 		FShaderSource S(SHADER_SOURCE_LITERAL("float4 PSMain() : SV_Target { return float4(1,0,0,1); };"));
 		auto P = ParseShader(S);
-		TestEqual(TEXT("Pixel shader entry point, total chunks"), P.Chunks.Num(), 1);
-		if (P.Chunks.Num() == 1)
+		TEST_IF_EQUAL(TEXT("Pixel shader entry point, total chunks"), P.Chunks.Num(), 1)
 		{
-			TestEqual(TEXT("Pixel shader entry point, main chunk type"), P.Chunks[0].Type, ECodeChunkType::Function);
+			TEST_EQUAL(TEXT("Pixel shader entry point, main chunk type"), P.Chunks[0].Type, ECodeChunkType::Function);
 		}
 	}
 
@@ -2180,11 +2176,11 @@ bool FShaderMinifierParserTest::RunTest(const FString& Parameters)
 		TArray<FShaderSource::FViewType, FMemStackAllocator> R;
 		FShaderSource S(SHADER_SOURCE_LITERAL("Hello[World]; Foo[0];\n"));
 		ExtractIdentifiers(S.GetView(), R);
-		if (TestEqual(TEXT("ExtractIdentifiers1: Num"), R.Num(), 3))
+		TEST_IF_EQUAL(TEXT("ExtractIdentifiers1: Num"), R.Num(), 3)
 		{
-			TestEqual(TEXT("ExtractIdentifiers1: R[0]"), FString(R[0]), SHADER_SOURCE_LITERAL("Hello"));
-			TestEqual(TEXT("ExtractIdentifiers1: R[1]"), FString(R[1]), SHADER_SOURCE_LITERAL("World"));
-			TestEqual(TEXT("ExtractIdentifiers1: R[2]"), FString(R[2]), SHADER_SOURCE_LITERAL("Foo"));
+			TEST_EQUAL(TEXT("ExtractIdentifiers1: R[0]"), FString(R[0]), SHADER_SOURCE_LITERAL("Hello"));
+			TEST_EQUAL(TEXT("ExtractIdentifiers1: R[1]"), FString(R[1]), SHADER_SOURCE_LITERAL("World"));
+			TEST_EQUAL(TEXT("ExtractIdentifiers1: R[2]"), FString(R[2]), SHADER_SOURCE_LITERAL("Foo"));
 		}
 	}
 
@@ -2193,64 +2189,64 @@ bool FShaderMinifierParserTest::RunTest(const FString& Parameters)
 		TArray<FShaderSource::FViewType, FMemStackAllocator> R;
 		FShaderSource S(SHADER_SOURCE_LITERAL("#line 0\nStructuredBuffer<uint4> Blah : register(t0, space123);#line 1\n#pragma foo\n"));
 		ExtractIdentifiers(S.GetView(), R);
-		if (TestEqual(TEXT("ExtractIdentifiers2: Num"), R.Num(), 6))
+		TEST_IF_EQUAL(TEXT("ExtractIdentifiers2: Num"), R.Num(), 6)
 		{
-			TestEqual(TEXT("ExtractIdentifiers2: R[0]"), FString(R[0]), SHADER_SOURCE_LITERAL("StructuredBuffer"));
-			TestEqual(TEXT("ExtractIdentifiers2: R[1]"), FString(R[1]), SHADER_SOURCE_LITERAL("uint4"));
-			TestEqual(TEXT("ExtractIdentifiers2: R[2]"), FString(R[2]), SHADER_SOURCE_LITERAL("Blah"));
-			TestEqual(TEXT("ExtractIdentifiers2: R[3]"), FString(R[3]), SHADER_SOURCE_LITERAL("register"));
-			TestEqual(TEXT("ExtractIdentifiers2: R[4]"), FString(R[4]), SHADER_SOURCE_LITERAL("t0"));
-			TestEqual(TEXT("ExtractIdentifiers2: R[5]"), FString(R[5]), SHADER_SOURCE_LITERAL("space123"));
+			TEST_EQUAL(TEXT("ExtractIdentifiers2: R[0]"), FString(R[0]), SHADER_SOURCE_LITERAL("StructuredBuffer"));
+			TEST_EQUAL(TEXT("ExtractIdentifiers2: R[1]"), FString(R[1]), SHADER_SOURCE_LITERAL("uint4"));
+			TEST_EQUAL(TEXT("ExtractIdentifiers2: R[2]"), FString(R[2]), SHADER_SOURCE_LITERAL("Blah"));
+			TEST_EQUAL(TEXT("ExtractIdentifiers2: R[3]"), FString(R[3]), SHADER_SOURCE_LITERAL("register"));
+			TEST_EQUAL(TEXT("ExtractIdentifiers2: R[4]"), FString(R[4]), SHADER_SOURCE_LITERAL("t0"));
+			TEST_EQUAL(TEXT("ExtractIdentifiers2: R[5]"), FString(R[5]), SHADER_SOURCE_LITERAL("space123"));
 		}
 	}
 
 	{
 		FShaderSource S(SHADER_SOURCE_LITERAL("StructuredBuffer<uint4> Blah : register(t0, space123);"));
 		auto P = ParseShader(S);
-		if (TestEqual(TEXT("ParseShader: structured buffer: num chunks"), P.Chunks.Num(), 1))
+		TEST_IF_EQUAL(TEXT("ParseShader: structured buffer: num chunks"), P.Chunks.Num(), 1)
 		{
-			TestEqual(TEXT("ParseShader: structured buffer: chunk"), P.Chunks[0].Type, ECodeChunkType::Variable);
+			TEST_EQUAL(TEXT("ParseShader: structured buffer: chunk"), P.Chunks[0].Type, ECodeChunkType::Variable);
 		}
 	}
 
 	{
 		FShaderSource S(SHADER_SOURCE_LITERAL("const float Foo = 123.45f;"));
 		auto P = ParseShader(S);
-		if (TestEqual(TEXT("ParseShader: const float with initializer: num chunks"), P.Chunks.Num(), 1))
+		TEST_IF_EQUAL(TEXT("ParseShader: const float with initializer: num chunks"), P.Chunks.Num(), 1)
 		{
-			TestEqual(TEXT("ParseShader: const float with initializer: chunk type"), P.Chunks[0].Type, ECodeChunkType::Variable);
+			TEST_EQUAL(TEXT("ParseShader: const float with initializer: chunk type"), P.Chunks[0].Type, ECodeChunkType::Variable);
 		}
 	}
 
 	{
 		FShaderSource S(SHADER_SOURCE_LITERAL("struct Blah { int A; };"));
 		auto P = ParseShader(S);
-		if (TestEqual(TEXT("ParseShader: struct: num chunks"), P.Chunks.Num(), 1))
+		TEST_IF_EQUAL(TEXT("ParseShader: struct: num chunks"), P.Chunks.Num(), 1)
 		{
-			TestEqual(TEXT("ParseShader: struct: chunk type"), P.Chunks[0].Type, ECodeChunkType::Struct);
+			TEST_EQUAL(TEXT("ParseShader: struct: chunk type"), P.Chunks[0].Type, ECodeChunkType::Struct);
 		}
 	}
 
 	{
 		FShaderSource S(SHADER_SOURCE_LITERAL("struct Foo { int FooA; }; struct Bar : Foo { int BarA; };"));
 		auto P = ParseShader(S);
-		if (TestEqual(TEXT("ParseShader: inherited struct: num chunks"), P.Chunks.Num(), 2))
+		TEST_IF_EQUAL(TEXT("ParseShader: inherited struct: num chunks"), P.Chunks.Num(), 2)
 		{
-			TestEqual(TEXT("ParseShader: inherited struct: chunk 0 type"), P.Chunks[0].Type, ECodeChunkType::Struct);
-			TestEqual(TEXT("ParseShader: inherited struct: chunk 1 type"), P.Chunks[1].Type, ECodeChunkType::Struct);
+			TEST_EQUAL(TEXT("ParseShader: inherited struct: chunk 0 type"), P.Chunks[0].Type, ECodeChunkType::Struct);
+			TEST_EQUAL(TEXT("ParseShader: inherited struct: chunk 1 type"), P.Chunks[1].Type, ECodeChunkType::Struct);
 		}
 	}
 
 	{
 		FShaderSource S(SHADER_SOURCE_LITERAL("[numthreads(8,8,1)] void Main() {};"));
 		auto P = ParseShader(S);
-		if (TestEqual(TEXT("ParseShader: compute shader entry point: num chunks"), P.Chunks.Num(), 1))
+		TEST_IF_EQUAL(TEXT("ParseShader: compute shader entry point: num chunks"), P.Chunks.Num(), 1)
 		{
-			TestEqual(TEXT("ParseShader: compute shader entry point: chunk type"), P.Chunks[0].Type, ECodeChunkType::Function);
+			TEST_EQUAL(TEXT("ParseShader: compute shader entry point: chunk type"), P.Chunks[0].Type, ECodeChunkType::Function);
 
-			if (TestEqual(TEXT("ParseShader: compute shader entry point: num blocks"), P.Chunks[0].Blocks.Num(), 5))
+			TEST_IF_EQUAL(TEXT("ParseShader: compute shader entry point: num blocks"), P.Chunks[0].Blocks.Num(), 5)
 			{
-				TestEqual(TEXT("ParseShader: compute shader entry point: attribute block type"), P.Chunks[0].Blocks[0].Type, EBlockType::Attribute);
+				TEST_EQUAL(TEXT("ParseShader: compute shader entry point: attribute block type"), P.Chunks[0].Blocks[0].Type, EBlockType::Attribute);
 			}
 		}
 	}
@@ -2258,26 +2254,26 @@ bool FShaderMinifierParserTest::RunTest(const FString& Parameters)
 	{
 		FShaderSource S(SHADER_SOURCE_LITERAL("Texture2D Blah : register(t0);"));
 		auto P = ParseShader(S);
-		if (TestEqual(TEXT("ParseShader: texture with register: num chunks"), P.Chunks.Num(), 1))
+		TEST_IF_EQUAL(TEXT("ParseShader: texture with register: num chunks"), P.Chunks.Num(), 1)
 		{
-			TestEqual(TEXT("ParseShader: texture with register: chunk type"), P.Chunks[0].Type, ECodeChunkType::Variable);
+			TEST_EQUAL(TEXT("ParseShader: texture with register: chunk type"), P.Chunks[0].Type, ECodeChunkType::Variable);
 		}
 	}
 	{
 		FShaderSource S(SHADER_SOURCE_LITERAL("Texture2D Blah;"));
 		auto P = ParseShader(S);
-		if (TestEqual(TEXT("ParseShader: texture: num chunks"), P.Chunks.Num(), 1))
+		TEST_IF_EQUAL(TEXT("ParseShader: texture: num chunks"), P.Chunks.Num(), 1)
 		{
-			TestEqual(TEXT("ParseShader: texture: chunk type"), P.Chunks[0].Type, ECodeChunkType::Variable);
+			TEST_EQUAL(TEXT("ParseShader: texture: chunk type"), P.Chunks[0].Type, ECodeChunkType::Variable);
 		}
 	}
 
 	{
 		FShaderSource S(SHADER_SOURCE_LITERAL("SamplerState Blah : register(s0, space123);"));
 		auto P = ParseShader(S);
-		if (TestEqual(TEXT("ParseShader: sampler state with register: num chunks"), P.Chunks.Num(), 1))
+		TEST_IF_EQUAL(TEXT("ParseShader: sampler state with register: num chunks"), P.Chunks.Num(), 1)
 		{
-			TestEqual(TEXT("ParseShader: sampler state with register: chunk type"), P.Chunks[0].Type, ECodeChunkType::Variable);
+			TEST_EQUAL(TEXT("ParseShader: sampler state with register: chunk type"), P.Chunks[0].Type, ECodeChunkType::Variable);
 		}
 	}
 
@@ -2286,25 +2282,25 @@ bool FShaderMinifierParserTest::RunTest(const FString& Parameters)
 		// TODO: handle function forward declarations
 		FShaderSource S(SHADER_SOURCE_LITERAL("Foo Fun(int a);"));
 		auto P = ParseShader(S);
-		TestEqual(TEXT("ParseShader: function forward declaration"), P.Chunks[0].Type, ECodeChunkType::FunctionDecl);
+		TEST_EQUAL(TEXT("ParseShader: function forward declaration"), P.Chunks[0].Type, ECodeChunkType::FunctionDecl);
 	}
 #endif
 
 	{
 		FShaderSource S(SHADER_SOURCE_LITERAL("void Fun(int a) {};"));
 		auto P = ParseShader(S);
-		if (TestEqual(TEXT("ParseShader: function with trailing semicolon: num chunks"), P.Chunks.Num(), 1))
+		TEST_IF_EQUAL(TEXT("ParseShader: function with trailing semicolon: num chunks"), P.Chunks.Num(), 1)
 		{
-			TestEqual(TEXT("ParseShader: function with trailing semicolon: chunk type"), P.Chunks[0].Type, ECodeChunkType::Function);
+			TEST_EQUAL(TEXT("ParseShader: function with trailing semicolon: chunk type"), P.Chunks[0].Type, ECodeChunkType::Function);
 		}
 	}
 
 	{
 		FShaderSource S(SHADER_SOURCE_LITERAL("void Fun(int a) {}"));
 		auto P = ParseShader(S);
-		if (TestEqual(TEXT("ParseShader: function: num chunks"), P.Chunks.Num(), 1))
+		TEST_IF_EQUAL(TEXT("ParseShader: function: num chunks"), P.Chunks.Num(), 1)
 		{
-			TestEqual(TEXT("ParseShader: function: chunk type"), P.Chunks[0].Type, ECodeChunkType::Function);
+			TEST_EQUAL(TEXT("ParseShader: function: chunk type"), P.Chunks[0].Type, ECodeChunkType::Function);
 		}
 	}
 
@@ -2312,28 +2308,28 @@ bool FShaderMinifierParserTest::RunTest(const FString& Parameters)
 		FShaderSource S(SHADER_SOURCE_LITERAL("cbuffer Foo {blah} SamplerState S;"));
 		auto P = ParseShader(S);
 
-		if (TestEqual(TEXT("ParseShader: cbuffer and sampler state: num chunks"), P.Chunks.Num(), 2))
+		TEST_IF_EQUAL(TEXT("ParseShader: cbuffer and sampler state: num chunks"), P.Chunks.Num(), 2)
 		{
-			TestEqual(TEXT("ParseShader: cbuffer and sampler state: chunk type [0]"), P.Chunks[0].Type, ECodeChunkType::CBuffer);
-			TestEqual(TEXT("ParseShader: cbuffer and sampler state: chunk type [1]"), P.Chunks[1].Type, ECodeChunkType::Variable);
+			TEST_EQUAL(TEXT("ParseShader: cbuffer and sampler state: chunk type [0]"), P.Chunks[0].Type, ECodeChunkType::CBuffer);
+			TEST_EQUAL(TEXT("ParseShader: cbuffer and sampler state: chunk type [1]"), P.Chunks[1].Type, ECodeChunkType::Variable);
 		}
 	}
 
 	{
 		FShaderSource S(SHADER_SOURCE_LITERAL("struct Foo { int a; };"));
 		auto P = ParseShader(S);
-		if (TestEqual(TEXT("ParseShader: struct: num chunks"), P.Chunks.Num(), 1))
+		TEST_IF_EQUAL(TEXT("ParseShader: struct: num chunks"), P.Chunks.Num(), 1)
 		{
-			TestEqual(TEXT("ParseShader: struct: chunk type"), P.Chunks[0].Type, ECodeChunkType::Struct);
+			TEST_EQUAL(TEXT("ParseShader: struct: chunk type"), P.Chunks[0].Type, ECodeChunkType::Struct);
 		}
 	}
 
 	{
 		FShaderSource S(SHADER_SOURCE_LITERAL("struct { int a; } Foo = { 123; };"));
 		auto P = ParseShader(S);
-		if (TestEqual(TEXT("ParseShader: anonymous struct with variable and initializer: num chunks"), P.Chunks.Num(), 1))
+		TEST_IF_EQUAL(TEXT("ParseShader: anonymous struct with variable and initializer: num chunks"), P.Chunks.Num(), 1)
 		{
-			TestEqual(TEXT("ParseShader: anonymous struct with variable and initializer: chunk type [0]"), P.Chunks[0].Type, ECodeChunkType::Variable);
+			TEST_EQUAL(TEXT("ParseShader: anonymous struct with variable and initializer: chunk type [0]"), P.Chunks[0].Type, ECodeChunkType::Variable);
 		}
 	}
 
@@ -2352,111 +2348,113 @@ bool FShaderMinifierParserTest::RunTest(const FString& Parameters)
 				"{ float4 Element1 : packoffset(c0); float1 Element2 : packoffset(c1); float1 Element3 : packoffset(c1.y); }")
 			);
 		auto P = ParseShader(S);
-		if (TestEqual(TEXT("ParseShader: cbuffer with packoffset: num chunks"), P.Chunks.Num(), 1))
+		TEST_IF_EQUAL(TEXT("ParseShader: cbuffer with packoffset: num chunks"), P.Chunks.Num(), 1)
 		{
-			TestEqual(TEXT("ParseShader: cbuffer with packoffset: chunk type"), P.Chunks[0].Type, ECodeChunkType::CBuffer);
+			TEST_EQUAL(TEXT("ParseShader: cbuffer with packoffset: chunk type"), P.Chunks[0].Type, ECodeChunkType::CBuffer);
 		}
 	}
 
 	{
 		FShaderSource S(SHADER_SOURCE_LITERAL("static const struct { float4 Param; } Foo;"));
 		auto P = ParseShader(S);
-		if (TestEqual(TEXT("ParseShader: static const anonymous struct with variable: num chunks"), P.Chunks.Num(), 1))
+		TEST_IF_EQUAL(TEXT("ParseShader: static const anonymous struct with variable: num chunks"), P.Chunks.Num(), 1)
 		{
-			TestEqual(TEXT("ParseShader: static const anonymous struct with variable: chunk type"), P.Chunks[0].Type, ECodeChunkType::Variable);
+			TEST_EQUAL(TEXT("ParseShader: static const anonymous struct with variable: chunk type"), P.Chunks[0].Type, ECodeChunkType::Variable);
 		}
 	}
 
 	{
 		FShaderSource S(SHADER_SOURCE_LITERAL("static const struct { float4 Param; } Foo = { FooCB_Param; };"));
 		auto P = ParseShader(S);
-		if (TestEqual(TEXT("ParseShader: static const anonymous struct with variable and initializer: num chunks"), P.Chunks.Num(), 1))
+		TEST_IF_EQUAL(TEXT("ParseShader: static const anonymous struct with variable and initializer: num chunks"), P.Chunks.Num(), 1)
 		{
-			TestEqual(TEXT("ParseShader: static const anonymous struct with variable and initializer: chunk type"), P.Chunks[0].Type, ECodeChunkType::Variable);
+			TEST_EQUAL(TEXT("ParseShader: static const anonymous struct with variable and initializer: chunk type"), P.Chunks[0].Type, ECodeChunkType::Variable);
 		}
 	}
 
 	{
 		FShaderSource S(SHADER_SOURCE_LITERAL("template <typename T> float Fun(T x) { return (float)x; }"));
 		auto P = ParseShader(S);
-		if (TestEqual(TEXT("ParseShader: template function: num chunks"), P.Chunks.Num(), 1))
+		TEST_IF_EQUAL(TEXT("ParseShader: template function: num chunks"), P.Chunks.Num(), 1)
 		{
-			TestEqual(TEXT("ParseShader: template function: chunk type"), P.Chunks[0].Type, ECodeChunkType::Function);
+			TEST_EQUAL(TEXT("ParseShader: template function: chunk type"), P.Chunks[0].Type, ECodeChunkType::Function);
 		}
 	}
 
 	{
 		FShaderSource S(SHADER_SOURCE_LITERAL("enum EFoo { A, B = 123 };"));
 		auto P = ParseShader(S);
-		if (TestEqual(TEXT("ParseShader: enum: num chunks"), P.Chunks.Num(), 1))
+		TEST_IF_EQUAL(TEXT("ParseShader: enum: num chunks"), P.Chunks.Num(), 1)
 		{
-			TestEqual(TEXT("ParseShader: enum: chunk type"), P.Chunks[0].Type, ECodeChunkType::Enum);
+			TEST_EQUAL(TEXT("ParseShader: enum: chunk type"), P.Chunks[0].Type, ECodeChunkType::Enum);
 		}
 	}
 
 	{
 		FShaderSource S(SHADER_SOURCE_LITERAL("enum class EFoo { A, B };"));
 		auto P = ParseShader(S);
-		if (TestEqual(TEXT("ParseShader: enum class: num chunks"), P.Chunks.Num(), 1))
+		TEST_IF_EQUAL(TEXT("ParseShader: enum class: num chunks"), P.Chunks.Num(), 1)
 		{
-			TestEqual(TEXT("ParseShader: enum class: chunk type"), P.Chunks[0].Type, ECodeChunkType::Enum);
+			TEST_EQUAL(TEXT("ParseShader: enum class: chunk type"), P.Chunks[0].Type, ECodeChunkType::Enum);
 		}
 	}
 
 	{
 		FShaderSource S(SHADER_SOURCE_LITERAL("#define Foo 123"));
 		auto P = ParseShader(S);
-		if (TestEqual(TEXT("ParseShader: define: num chunks"), P.Chunks.Num(), 1))
+		TEST_IF_EQUAL(TEXT("ParseShader: define: num chunks"), P.Chunks.Num(), 1)
 		{
-			TestEqual(TEXT("ParseShader: define: chunk type"), P.Chunks[0].Type, ECodeChunkType::Define);
+			TEST_EQUAL(TEXT("ParseShader: define: chunk type"), P.Chunks[0].Type, ECodeChunkType::Define);
 		}
 	}
 
 	{
 		FShaderSource S(SHADER_SOURCE_LITERAL("#pragma Foo"));
 		auto P = ParseShader(S);
-		if (TestEqual(TEXT("ParseShader: pragma: num chunks"), P.Chunks.Num(), 1))
+		TEST_IF_EQUAL(TEXT("ParseShader: pragma: num chunks"), P.Chunks.Num(), 1)
 		{
-			TestEqual(TEXT("ParseShader: pragma: chunk type"), P.Chunks[0].Type, ECodeChunkType::Pragma);
+			TEST_EQUAL(TEXT("ParseShader: pragma: chunk type"), P.Chunks[0].Type, ECodeChunkType::Pragma);
 		}
 	}
 
 	{
 		FShaderSource S(SHADER_SOURCE_LITERAL("ConstantBuffer<Foo> CB : register ( b123, space456);"));
 		auto P = ParseShader(S);
-		if (TestEqual(TEXT("ParseShader: ConstantBuffer<Foo>: num chunks"), P.Chunks.Num(), 1))
+		TEST_IF_EQUAL(TEXT("ParseShader: ConstantBuffer<Foo>: num chunks"), P.Chunks.Num(), 1)
 		{
-			TestEqual(TEXT("ParseShader: ConstantBuffer<Foo>: chunk type"), P.Chunks[0].Type, ECodeChunkType::Variable);
+			TEST_EQUAL(TEXT("ParseShader: ConstantBuffer<Foo>: chunk type"), P.Chunks[0].Type, ECodeChunkType::Variable);
 		}
 	}
 
 	{
 		FShaderSource S(SHADER_SOURCE_LITERAL("namespace NS1 { void Fun() {}; } namespace NS2 { void Fun() {}; }"));
 		auto P = ParseShader(S);
-		if (TestEqual(TEXT("ParseShader: namespaces: num chunks"), P.Chunks.Num(), 2)
-			&& TestEqual(TEXT("ParseShader: namespaces: num namespaces"), P.Namespaces.Num(), 2))
+		TEST_IF_EQUAL(TEXT("ParseShader: namespaces: num chunks"), P.Chunks.Num(), 2)
 		{
-			TestEqual(TEXT("ParseShader: namespaces: chunk 0 namespace"), P.Chunks[0].Namespace, 0);
-			TestEqual(TEXT("ParseShader: namespaces: chunk 1 namespace"), P.Chunks[1].Namespace, 1);
+			TEST_IF_EQUAL(TEXT("ParseShader: namespaces: num namespaces"), P.Namespaces.Num(), 2)
+			{
+				TEST_EQUAL(TEXT("ParseShader: namespaces: chunk 0 namespace"), P.Chunks[0].Namespace, 0);
+				TEST_EQUAL(TEXT("ParseShader: namespaces: chunk 1 namespace"), P.Chunks[1].Namespace, 1);
+			}
 		}
 	}
 
 	{
 		FShaderSource S(SHADER_SOURCE_LITERAL("template< typename T > TMyStruct<T> operator + ( TMyStruct<T> A, T B ) { /*...*/ }"));
 		auto P = ParseShader(S);
-		if (TestEqual(TEXT("ParseShader: operators: num chunks"), P.Chunks.Num(), 1))
+		TEST_IF_EQUAL(TEXT("ParseShader: operators: num chunks"), P.Chunks.Num(), 1)
 		{
-			TestEqual(TEXT("ParseShader: operators: chunk type"), P.Chunks[0].Type, ECodeChunkType::Operator);
-			if (TestEqual(TEXT("ParseShader: operators: chunk 0: num blocks"), P.Chunks[0].Blocks.Num(), 8))
+			TEST_EQUAL(TEXT("ParseShader: operators: chunk type"), P.Chunks[0].Type, ECodeChunkType::Operator);
+			TEST_IF_EQUAL(TEXT("ParseShader: operators: chunk 0: num blocks"), P.Chunks[0].Blocks.Num(), 8)
 			{
-				FString Args = FString(P.Chunks[0].FindFirstBlockByType(EBlockType::Args));
-				TestEqual(TEXT("ParseShader: operators: chunk 0: type name"), *Args, SHADER_SOURCE_LITERAL("( TMyStruct<T> A, T B )"));
+				FShaderSource::FStringType Args = FShaderSource::FStringType(P.Chunks[0].FindFirstBlockByType(EBlockType::Args));
+				TEST_EQUAL(TEXT("ParseShader: operators: chunk 0: type name"), Args, SHADER_SOURCE_LITERAL("( TMyStruct<T> A, T B )"));
 
-				FString TypeName = FString(P.Chunks[0].FindFirstBlockByType(EBlockType::Type));
-				TestEqual(TEXT("ParseShader: operators: chunk 0: type name"), *TypeName, SHADER_SOURCE_LITERAL("TMyStruct"));
+				FShaderSource::FStringType TypeName = FShaderSource::FStringType(P.Chunks[0].FindFirstBlockByType(EBlockType::Type));
+				TEST_EQUAL(TEXT("ParseShader: operators: chunk 0: type name"), TypeName, SHADER_SOURCE_LITERAL("TMyStruct"));
 
-				FString OperatorName = FString(P.Chunks[0].FindFirstBlockByType(EBlockType::OperatorName));
-				TestEqual(TEXT("ParseShader: operators: chunk 0: operator name"), *OperatorName, SHADER_SOURCE_LITERAL("+"));
+				FShaderSource::FStringType OperatorName = FShaderSource::FStringType(P.Chunks[0].FindFirstBlockByType(EBlockType::OperatorName));
+				TEST_EQUAL(TEXT("ParseShader: operators: chunk 0: operator name"), OperatorName, SHADER_SOURCE_LITERAL("+"));
 			}
 		}
 	}
@@ -2464,44 +2462,37 @@ bool FShaderMinifierParserTest::RunTest(const FString& Parameters)
 	{
 		FShaderSource S(SHADER_SOURCE_LITERAL("typedef Bar Foo;"));
 		auto P = ParseShader(S);
-		if (TestEqual(TEXT("ParseShader: standard typedef: num chunks"), P.Chunks.Num(), 1))
+		TEST_IF_EQUAL(TEXT("ParseShader: standard typedef: num chunks"), P.Chunks.Num(), 1)
 		{
-			TestEqual(TEXT("ParseShader: standard typedef: chunk type"), P.Chunks[0].Type, ECodeChunkType::Typedef);
+			TEST_EQUAL(TEXT("ParseShader: standard typedef: chunk type"), P.Chunks[0].Type, ECodeChunkType::Typedef);
 		}
 	}
 
 	{
 		FShaderSource S(SHADER_SOURCE_LITERAL("typedef Bar Foo; static const Foo = Bar(0);"));
 		auto P = ParseShader(S);
-		if (TestEqual(TEXT("ParseShader: standard typedef: num chunks"), P.Chunks.Num(), 2))
+		TEST_IF_EQUAL(TEXT("ParseShader: standard typedef: num chunks"), P.Chunks.Num(), 2)
 		{
-			TestEqual(TEXT("ParseShader: standard typedef: chunk type"), P.Chunks[0].Type, ECodeChunkType::Typedef);
-			TestEqual(TEXT("ParseShader: standard typedef: chunk type"), P.Chunks[1].Type, ECodeChunkType::Variable);
+			TEST_EQUAL(TEXT("ParseShader: standard typedef: chunk type"), P.Chunks[0].Type, ECodeChunkType::Typedef);
+			TEST_EQUAL(TEXT("ParseShader: standard typedef: chunk type"), P.Chunks[1].Type, ECodeChunkType::Variable);
 		}
 	}
 
 	{
 		FShaderSource S(SHADER_SOURCE_LITERAL("[Shader(\"node\")] [NodeID(\"WorkgraphNodeArray\", 0)] void WorkgraphNodeArray_0() {}"));
 		auto P = ParseShader(S);
-		TestEqual(TEXT("ParseShader: workgraph node: num chunks"), P.Chunks.Num(), 1);
-		if (P.Chunks.Num() == 1)
+		TEST_IF_EQUAL(TEXT("ParseShader: workgraph node: num chunks"), P.Chunks.Num(), 1)
 		{
-			TestEqual(TEXT("ParseShader: workgraph node, chunk type"), P.Chunks[0].Type, ECodeChunkType::Function);
-			if (TestEqual(TEXT("ParseShader: workgraph node: chunk 0: num blocks"), P.Chunks[0].Blocks.Num(), 6))
+			TEST_EQUAL(TEXT("ParseShader: workgraph node, chunk type"), P.Chunks[0].Type, ECodeChunkType::Function);
+			TEST_IF_EQUAL(TEXT("ParseShader: workgraph node: chunk 0: num blocks"), P.Chunks[0].Blocks.Num(), 6)
 			{
-				TestEqual(TEXT("ParseShader: workgraph node: chunk 0: block 1: block type"), P.Chunks[0].Blocks[1].Type, EBlockType::NodeId);
+				TEST_EQUAL(TEXT("ParseShader: workgraph node: chunk 0: block 1: block type"), P.Chunks[0].Blocks[1].Type, EBlockType::NodeId);
 			}
 		}
 	}
-
-	int32 NumErrors = ExecutionInfo.GetErrorTotal();
-
-	return NumErrors == 0;
 }
 
-IMPLEMENT_SIMPLE_AUTOMATION_TEST(FShaderMinifierTest, "System.Shaders.ShaderMinifier.Minify", EAutomationTestFlags::EditorContext | EAutomationTestFlags::EngineFilter);
-
-bool FShaderMinifierTest::RunTest(const FString& Parameters)
+TEST_CASE_NAMED(FShaderMinifierMinifyTest, "Shaders::ShaderMinifier::Minify", "[EditorContext][EngineFilter]")
 {
 	using namespace UE::ShaderMinifier;
 
@@ -2712,11 +2703,11 @@ void WorkgraphNodeArray_1()
 		FDiagnostics Diagnostics;
 		FShaderSource Minified = FShaderSource(MinifyShader(Parsed, SHADER_SOURCE_LITERAL("EmptyFunction"), EMinifyShaderFlags::None, Diagnostics));
 		FParsedShader MinifiedParsed = ParseShader(Minified);
-		if (TestEqual(TEXT("MinifyShader: EmptyFunction: num chunks"), MinifiedParsed.Chunks.Num(), 3))
+		TEST_IF_EQUAL(TEXT("MinifyShader: EmptyFunction: num chunks"), MinifiedParsed.Chunks.Num(), 3)
 		{
-			TestEqual(TEXT("MinifyShader: EmptyFunction: pragma"), *FString(MinifiedParsed.Chunks[0].GetCode()), SHADER_SOURCE_LITERAL("#pragma test_pragma"));
-			TestEqual(TEXT("MinifyShader: EmptyFunction: define"), *FString(MinifiedParsed.Chunks[1].GetCode()), SHADER_SOURCE_LITERAL("#define COMPILER_DEFINITION_TEST 123"));
-			TestEqual(TEXT("MinifyShader: EmptyFunction: function"), *FString(MinifiedParsed.Chunks[2].GetCode()), SHADER_SOURCE_LITERAL("void EmptyFunction(){}"));
+			TEST_EQUAL(TEXT("MinifyShader: EmptyFunction: pragma"), MinifiedParsed.Chunks[0].GetCode(), SHADER_SOURCE_LITERAL("#pragma test_pragma"));
+			TEST_EQUAL(TEXT("MinifyShader: EmptyFunction: define"), MinifiedParsed.Chunks[1].GetCode(), SHADER_SOURCE_LITERAL("#define COMPILER_DEFINITION_TEST 123"));
+			TEST_EQUAL(TEXT("MinifyShader: EmptyFunction: function"), MinifiedParsed.Chunks[2].GetCode(), SHADER_SOURCE_LITERAL("void EmptyFunction(){}"));
 		}
 	}
 
@@ -2726,54 +2717,50 @@ void WorkgraphNodeArray_1()
 		FParsedShader MinifiedParsed = ParseShader(Minified);
 
 		// Expect true:
-		TestTrue(TEXT("MinifyShader: MainCS: contains MainCS"), ChunkPresent(MinifiedParsed, SHADER_SOURCE_LITERAL("MainCS")));
-		TestTrue(TEXT("MinifyShader: MainCS: contains FFoo"), ChunkPresent(MinifiedParsed, SHADER_SOURCE_LITERAL("FFoo")));
-		TestTrue(TEXT("MinifyShader: MainCS: contains FBar"), ChunkPresent(MinifiedParsed, SHADER_SOURCE_LITERAL("FBar")));
-		TestTrue(TEXT("MinifyShader: MainCS: contains Sum"), ChunkPresent(MinifiedParsed, SHADER_SOURCE_LITERAL("Sum")));
-		TestTrue(TEXT("MinifyShader: MainCS: contains FunA"), ChunkPresent(MinifiedParsed, SHADER_SOURCE_LITERAL("FunA")));
-		TestTrue(TEXT("MinifyShader: MainCS: contains FunB"), ChunkPresent(MinifiedParsed, SHADER_SOURCE_LITERAL("FunB")));
-		TestTrue(TEXT("MinifyShader: MainCS: contains GAnonymousStruct"), ChunkPresent(MinifiedParsed, SHADER_SOURCE_LITERAL("GAnonymousStruct")));
-		TestTrue(TEXT("MinifyShader: MainCS: contains GStructA"), ChunkPresent(MinifiedParsed, SHADER_SOURCE_LITERAL("GStructA")));
-		TestTrue(TEXT("MinifyShader: MainCS: contains GStructB"), ChunkPresent(MinifiedParsed, SHADER_SOURCE_LITERAL("GStructB")));
-		TestTrue(TEXT("MinifyShader: MainCS: contains GStructC"), ChunkPresent(MinifiedParsed, SHADER_SOURCE_LITERAL("GStructC")));
-		TestTrue(TEXT("MinifyShader: MainCS: contains GInitializedAnonymousStructA"), ChunkPresent(MinifiedParsed, SHADER_SOURCE_LITERAL("GInitializedAnonymousStructA")));
-		TestTrue(TEXT("MinifyShader: MainCS: contains GInitializedAnonymousStructB"), ChunkPresent(MinifiedParsed, SHADER_SOURCE_LITERAL("GInitializedAnonymousStructB")));
-		TestTrue(TEXT("MinifyShader: MainCS: contains OutputBufferType"), ChunkPresent(MinifiedParsed, SHADER_SOURCE_LITERAL("OutputBufferType")));
-		TestTrue(TEXT("MinifyShader: MainCS: contains OutputBuffer"), ChunkPresent(MinifiedParsed, SHADER_SOURCE_LITERAL("OutputBuffer")));
-		TestTrue(TEXT("MinifyShader: MainCS: contains FTypedefUsedStruct"), ChunkPresent(MinifiedParsed, SHADER_SOURCE_LITERAL("FTypedefUsedStruct")));
-		TestTrue(TEXT("MinifyShader: MainCS: contains FTypedefUsed"), ChunkPresent(MinifiedParsed, SHADER_SOURCE_LITERAL("FTypedefUsed")));
-		TestTrue(TEXT("MinifyShader: MainCS: contains FTypedefUsedChained"), ChunkPresent(MinifiedParsed, SHADER_SOURCE_LITERAL("FTypedefUsedChained")));
-		TestTrue(TEXT("MinifyShader: MainCS: contains TypedefUsedBuffer"), ChunkPresent(MinifiedParsed, SHADER_SOURCE_LITERAL("TypedefUsedBuffer")));
-		TestTrue(TEXT("MinifyShader: MainCS: contains struct TUnusedTemplate"), MinifiedParsed.Source.Contains(SHADER_SOURCE_LITERAL("struct TUsedTemplate")));
-		TestTrue(TEXT("MinifyShader: MainCS: contains TUsedTemplate<T> operator*"), MinifiedParsed.Source.Contains(SHADER_SOURCE_LITERAL("TUsedTemplate<T> operator*")));
-		TestTrue(TEXT("MinifyShader: MainCS: contains EEnumUsed"), MinifiedParsed.Source.Contains(SHADER_SOURCE_LITERAL("EEnumUsed")));
-		TestTrue(TEXT("MinifyShader: MainCS: contains ENUM_USED_PART_1"), MinifiedParsed.Source.Contains(SHADER_SOURCE_LITERAL("ENUM_USED_PART_1")));
-		TestTrue(TEXT("MinifyShader: MainCS: contains ENUM_USED_PART_2"), MinifiedParsed.Source.Contains(SHADER_SOURCE_LITERAL("ENUM_USED_PART_2")));
-		TestTrue(TEXT("MinifyShader: MainCS: contains struct FWorkgraphRecord"), MinifiedParsed.Source.Contains(SHADER_SOURCE_LITERAL("struct FWorkgraphRecord")));
-		TestTrue(TEXT("MinifyShader: MainCS: contains void WorkgraphNode()"), MinifiedParsed.Source.Contains(SHADER_SOURCE_LITERAL("void WorkgraphNode()")));
-		TestTrue(TEXT("MinifyShader: MainCS: contains void WorkgraphNodeFunction()"), MinifiedParsed.Source.Contains(SHADER_SOURCE_LITERAL("void WorkgraphNodeFunction()")));
-		TestTrue(TEXT("MinifyShader: MainCS: contains void WorkgraphNodeArray_0()"), MinifiedParsed.Source.Contains(SHADER_SOURCE_LITERAL("void WorkgraphNodeArray_0()")));
-		TestTrue(TEXT("MinifyShader: MainCS: contains void WorkgraphNodeArray_1()"), MinifiedParsed.Source.Contains(SHADER_SOURCE_LITERAL("void WorkgraphNodeArray_1()")));
+		TEST_TRUE(TEXT("MinifyShader: MainCS: contains MainCS"), ChunkPresent(MinifiedParsed, SHADER_SOURCE_LITERAL("MainCS")));
+		TEST_TRUE(TEXT("MinifyShader: MainCS: contains FFoo"), ChunkPresent(MinifiedParsed, SHADER_SOURCE_LITERAL("FFoo")));
+		TEST_TRUE(TEXT("MinifyShader: MainCS: contains FBar"), ChunkPresent(MinifiedParsed, SHADER_SOURCE_LITERAL("FBar")));
+		TEST_TRUE(TEXT("MinifyShader: MainCS: contains Sum"), ChunkPresent(MinifiedParsed, SHADER_SOURCE_LITERAL("Sum")));
+		TEST_TRUE(TEXT("MinifyShader: MainCS: contains FunA"), ChunkPresent(MinifiedParsed, SHADER_SOURCE_LITERAL("FunA")));
+		TEST_TRUE(TEXT("MinifyShader: MainCS: contains FunB"), ChunkPresent(MinifiedParsed, SHADER_SOURCE_LITERAL("FunB")));
+		TEST_TRUE(TEXT("MinifyShader: MainCS: contains GAnonymousStruct"), ChunkPresent(MinifiedParsed, SHADER_SOURCE_LITERAL("GAnonymousStruct")));
+		TEST_TRUE(TEXT("MinifyShader: MainCS: contains GStructA"), ChunkPresent(MinifiedParsed, SHADER_SOURCE_LITERAL("GStructA")));
+		TEST_TRUE(TEXT("MinifyShader: MainCS: contains GStructB"), ChunkPresent(MinifiedParsed, SHADER_SOURCE_LITERAL("GStructB")));
+		TEST_TRUE(TEXT("MinifyShader: MainCS: contains GStructC"), ChunkPresent(MinifiedParsed, SHADER_SOURCE_LITERAL("GStructC")));
+		TEST_TRUE(TEXT("MinifyShader: MainCS: contains GInitializedAnonymousStructA"), ChunkPresent(MinifiedParsed, SHADER_SOURCE_LITERAL("GInitializedAnonymousStructA")));
+		TEST_TRUE(TEXT("MinifyShader: MainCS: contains GInitializedAnonymousStructB"), ChunkPresent(MinifiedParsed, SHADER_SOURCE_LITERAL("GInitializedAnonymousStructB")));
+		TEST_TRUE(TEXT("MinifyShader: MainCS: contains OutputBufferType"), ChunkPresent(MinifiedParsed, SHADER_SOURCE_LITERAL("OutputBufferType")));
+		TEST_TRUE(TEXT("MinifyShader: MainCS: contains OutputBuffer"), ChunkPresent(MinifiedParsed, SHADER_SOURCE_LITERAL("OutputBuffer")));
+		TEST_TRUE(TEXT("MinifyShader: MainCS: contains FTypedefUsedStruct"), ChunkPresent(MinifiedParsed, SHADER_SOURCE_LITERAL("FTypedefUsedStruct")));
+		TEST_TRUE(TEXT("MinifyShader: MainCS: contains FTypedefUsed"), ChunkPresent(MinifiedParsed, SHADER_SOURCE_LITERAL("FTypedefUsed")));
+		TEST_TRUE(TEXT("MinifyShader: MainCS: contains FTypedefUsedChained"), ChunkPresent(MinifiedParsed, SHADER_SOURCE_LITERAL("FTypedefUsedChained")));
+		TEST_TRUE(TEXT("MinifyShader: MainCS: contains TypedefUsedBuffer"), ChunkPresent(MinifiedParsed, SHADER_SOURCE_LITERAL("TypedefUsedBuffer")));
+		TEST_TRUE(TEXT("MinifyShader: MainCS: contains struct TUnusedTemplate"), MinifiedParsed.Source.Contains(SHADER_SOURCE_LITERAL("struct TUsedTemplate")));
+		TEST_TRUE(TEXT("MinifyShader: MainCS: contains TUsedTemplate<T> operator*"), MinifiedParsed.Source.Contains(SHADER_SOURCE_LITERAL("TUsedTemplate<T> operator*")));
+		TEST_TRUE(TEXT("MinifyShader: MainCS: contains EEnumUsed"), MinifiedParsed.Source.Contains(SHADER_SOURCE_LITERAL("EEnumUsed")));
+		TEST_TRUE(TEXT("MinifyShader: MainCS: contains ENUM_USED_PART_1"), MinifiedParsed.Source.Contains(SHADER_SOURCE_LITERAL("ENUM_USED_PART_1")));
+		TEST_TRUE(TEXT("MinifyShader: MainCS: contains ENUM_USED_PART_2"), MinifiedParsed.Source.Contains(SHADER_SOURCE_LITERAL("ENUM_USED_PART_2")));
+		TEST_TRUE(TEXT("MinifyShader: MainCS: contains struct FWorkgraphRecord"), MinifiedParsed.Source.Contains(SHADER_SOURCE_LITERAL("struct FWorkgraphRecord")));
+		TEST_TRUE(TEXT("MinifyShader: MainCS: contains void WorkgraphNode()"), MinifiedParsed.Source.Contains(SHADER_SOURCE_LITERAL("void WorkgraphNode()")));
+		TEST_TRUE(TEXT("MinifyShader: MainCS: contains void WorkgraphNodeFunction()"), MinifiedParsed.Source.Contains(SHADER_SOURCE_LITERAL("void WorkgraphNodeFunction()")));
+		TEST_TRUE(TEXT("MinifyShader: MainCS: contains void WorkgraphNodeArray_0()"), MinifiedParsed.Source.Contains(SHADER_SOURCE_LITERAL("void WorkgraphNodeArray_0()")));
+		TEST_TRUE(TEXT("MinifyShader: MainCS: contains void WorkgraphNodeArray_1()"), MinifiedParsed.Source.Contains(SHADER_SOURCE_LITERAL("void WorkgraphNodeArray_1()")));
 
 		// Expect false:
-		TestFalse(TEXT("MinifyShader: MainCS: contains UnreferencedFunction"), ChunkPresent(MinifiedParsed, SHADER_SOURCE_LITERAL("UnreferencedFunction")));
-		TestFalse(TEXT("MinifyShader: MainCS: contains FUnreferencedStruct"), ChunkPresent(MinifiedParsed, SHADER_SOURCE_LITERAL("FUnreferencedStruct")));
-		TestFalse(TEXT("MinifyShader: MainCS: contains GUnreferencedParameter"), ChunkPresent(MinifiedParsed, SHADER_SOURCE_LITERAL("GUnreferencedParameter")));
-		TestFalse(TEXT("MinifyShader: MainCS: contains FTypedefUsedChainedUnused"), ChunkPresent(MinifiedParsed, SHADER_SOURCE_LITERAL("FTypedefUsedChainedUnused")));
-		TestFalse(TEXT("MinifyShader: MainCS: contains FTypedefUnusedStruct"), ChunkPresent(MinifiedParsed, SHADER_SOURCE_LITERAL("FTypedefUnusedStruct")));
-		TestFalse(TEXT("MinifyShader: MainCS: contains FTypedefUnused"), ChunkPresent(MinifiedParsed, SHADER_SOURCE_LITERAL("FTypedefUnused")));
-		TestFalse(TEXT("MinifyShader: MainCS: contains TypedefUnusedBuffer"), ChunkPresent(MinifiedParsed, SHADER_SOURCE_LITERAL("TypedefUnusedBuffer")));
-		TestFalse(TEXT("MinifyShader: MainCS: contains struct TUnusedTemplate"), MinifiedParsed.Source.Contains(SHADER_SOURCE_LITERAL("struct TUnusedTemplate")));
-		TestFalse(TEXT("MinifyShader: MainCS: contains TUnusedTemplate<T> operator%"), MinifiedParsed.Source.Contains(SHADER_SOURCE_LITERAL("TUnusedTemplate<T> operator%")));
-		TestFalse(TEXT("MinifyShader: MainCS: contains EEnumUnused"), MinifiedParsed.Source.Contains(SHADER_SOURCE_LITERAL("EEnumUnused")));
-		TestFalse(TEXT("MinifyShader: MainCS: contains ENUM_UNUSED_PART_1"), MinifiedParsed.Source.Contains(SHADER_SOURCE_LITERAL("ENUM_UNUSED_PART_1")));
-		TestFalse(TEXT("MinifyShader: MainCS: contains ENUM_UNUSED_PART_2"), MinifiedParsed.Source.Contains(SHADER_SOURCE_LITERAL("ENUM_UNUSED_PART_2")));
-		TestFalse(TEXT("MinifyShader: MainCS: contains UnreferencedEntryPoint"), MinifiedParsed.Source.Contains(SHADER_SOURCE_LITERAL("UnreferencedEntryPoint")));
-	}
-
-	int32 NumErrors = ExecutionInfo.GetErrorTotal();
-
-	return NumErrors == 0;
+		TEST_FALSE(TEXT("MinifyShader: MainCS: contains UnreferencedFunction"), ChunkPresent(MinifiedParsed, SHADER_SOURCE_LITERAL("UnreferencedFunction")));
+		TEST_FALSE(TEXT("MinifyShader: MainCS: contains FUnreferencedStruct"), ChunkPresent(MinifiedParsed, SHADER_SOURCE_LITERAL("FUnreferencedStruct")));
+		TEST_FALSE(TEXT("MinifyShader: MainCS: contains GUnreferencedParameter"), ChunkPresent(MinifiedParsed, SHADER_SOURCE_LITERAL("GUnreferencedParameter")));
+		TEST_FALSE(TEXT("MinifyShader: MainCS: contains FTypedefUsedChainedUnused"), ChunkPresent(MinifiedParsed, SHADER_SOURCE_LITERAL("FTypedefUsedChainedUnused")));
+		TEST_FALSE(TEXT("MinifyShader: MainCS: contains FTypedefUnusedStruct"), ChunkPresent(MinifiedParsed, SHADER_SOURCE_LITERAL("FTypedefUnusedStruct")));
+		TEST_FALSE(TEXT("MinifyShader: MainCS: contains FTypedefUnused"), ChunkPresent(MinifiedParsed, SHADER_SOURCE_LITERAL("FTypedefUnused")));
+		TEST_FALSE(TEXT("MinifyShader: MainCS: contains TypedefUnusedBuffer"), ChunkPresent(MinifiedParsed, SHADER_SOURCE_LITERAL("TypedefUnusedBuffer")));
+		TEST_FALSE(TEXT("MinifyShader: MainCS: contains struct TUnusedTemplate"), MinifiedParsed.Source.Contains(SHADER_SOURCE_LITERAL("struct TUnusedTemplate")));
+		TEST_FALSE(TEXT("MinifyShader: MainCS: contains TUnusedTemplate<T> operator%"), MinifiedParsed.Source.Contains(SHADER_SOURCE_LITERAL("TUnusedTemplate<T> operator%")));
+		TEST_FALSE(TEXT("MinifyShader: MainCS: contains EEnumUnused"), MinifiedParsed.Source.Contains(SHADER_SOURCE_LITERAL("EEnumUnused")));
+		TEST_FALSE(TEXT("MinifyShader: MainCS: contains ENUM_UNUSED_PART_1"), MinifiedParsed.Source.Contains(SHADER_SOURCE_LITERAL("ENUM_UNUSED_PART_1")));
+		TEST_FALSE(TEXT("MinifyShader: MainCS: contains ENUM_UNUSED_PART_2"), MinifiedParsed.Source.Contains(SHADER_SOURCE_LITERAL("ENUM_UNUSED_PART_2")));
+		TEST_FALSE(TEXT("MinifyShader: MainCS: contains UnreferencedEntryPoint"), MinifiedParsed.Source.Contains(SHADER_SOURCE_LITERAL("UnreferencedEntryPoint")));
+	}
 }
 
-#endif // WITH_AUTOMATION_TESTS
+#endif // WITH_LOW_LEVEL_TESTS
diff --git a/Engine/Source/Programs/LowLevelTests/FoundationTests/FoundationTests.Build.cs b/Engine/Source/Programs/LowLevelTests/FoundationTests/FoundationTests.Build.cs
--- a/Engine/Source/Programs/LowLevelTests/FoundationTests/FoundationTests.Build.cs
+++ b/Engine/Source/Programs/LowLevelTests/FoundationTests/FoundationTests.Build.cs
@@ -56,6 +56,7 @@ public class FoundationTests : TestModuleRules
 			PrivateDependencyModuleNames.AddRange(
 				new string[] {
 					"DesktopPlatform",
+					"ShaderCompilerCommon",
 					"ShaderPreprocessor"
 				});
 		}
