[shaders] add handling of c99/c++11-style _Pragma directives to minifier & update unit test to actually test this properly

diff --git a/Engine/Source/Developer/ShaderCompilerCommon/Private/ShaderMinifier.cpp b/Engine/Source/Developer/ShaderCompilerCommon/Private/ShaderMinifier.cpp
--- a/Engine/Source/Developer/ShaderCompilerCommon/Private/ShaderMinifier.cpp
+++ b/Engine/Source/Developer/ShaderCompilerCommon/Private/ShaderMinifier.cpp
@@ -118,7 +118,7 @@ struct FCharacterFlags
 			Flags[uint8(C)] |= Letter;
 		}
 
-		Flags[uint8('_')] |= Underscore;
+		Flags[uint8('_')] |= (Underscore | Special);
 
 		Flags[uint8(' ')]  |= Space;
 		Flags[uint8('\f')] |= Space;
@@ -969,6 +969,36 @@ static FParsedShader ParseShader(const FShaderSource& InSource, FDiagnostics& Ou
 					continue;
 				}
 			}
+			else if (FirstChar == '_' && StartsWith(Source, SHADER_SOURCE_VIEWLITERAL("_Pragma")))
+			{
+				FShaderSource::FViewType PragmaStart = Source;
+				FShaderSource::FViewType Remainder = Source;
+
+				// handle case of multiple _Pragma operators on the same line (to avoid inserting newlines)
+				while (PragmaStart.Len() > 0)
+				{
+					Remainder = SkipUntilStr(Remainder, SHADER_SOURCE_VIEWLITERAL(")"));
+					if (Remainder.Len() >= 1)
+					{
+						Remainder = SubStrView(Remainder, 1);
+					}
+
+					FShaderSource::FViewType NextLine = SkipUntilNextLine(Remainder);
+					FShaderSource::FViewType LineRemainder = SubStrView(Remainder, 0, Remainder.Len() - NextLine.Len());
+					PragmaStart = SkipUntilStr(LineRemainder, SHADER_SOURCE_VIEWLITERAL("_Pragma"));
+					if (PragmaStart.Len() > 0)
+					{
+						Remainder = SkipUntilStr(Remainder, SHADER_SOURCE_VIEWLITERAL("_Pragma"));
+					}
+				}
+
+				FShaderSource::FViewType Block = SubStrView(Source, 0, Source.Len() - Remainder.Len());
+				AddBlock(EBlockType::Directive, Block);
+				ChunkType = ECodeChunkType::Pragma;
+				FinalizeChunk();
+				Source = Remainder;
+				continue;
+			}
 			else if (PendingBlocks.IsEmpty() && (FirstChar == '{' || FirstChar == '}'))
 			{
 				if (StartsWith(Source, SHADER_SOURCE_VIEWLITERAL("{")))
@@ -2703,11 +2733,15 @@ void WorkgraphNodeArray_1()
 		FDiagnostics Diagnostics;
 		FShaderSource Minified = FShaderSource(MinifyShader(Parsed, SHADER_SOURCE_LITERAL("EmptyFunction"), EMinifyShaderFlags::None, Diagnostics));
 		FParsedShader MinifiedParsed = ParseShader(Minified);
-		TEST_IF_EQUAL(TEXT("MinifyShader: EmptyFunction: num chunks"), MinifiedParsed.Chunks.Num(), 3)
+		TEST_IF_EQUAL(TEXT("MinifyShader: EmptyFunction: num chunks"), MinifiedParsed.Chunks.Num(), 7)
 		{
 			TEST_EQUAL(TEXT("MinifyShader: EmptyFunction: pragma"), MinifiedParsed.Chunks[0].GetCode(), SHADER_SOURCE_LITERAL("#pragma test_pragma"));
-			TEST_EQUAL(TEXT("MinifyShader: EmptyFunction: define"), MinifiedParsed.Chunks[1].GetCode(), SHADER_SOURCE_LITERAL("#define COMPILER_DEFINITION_TEST 123"));
-			TEST_EQUAL(TEXT("MinifyShader: EmptyFunction: function"), MinifiedParsed.Chunks[2].GetCode(), SHADER_SOURCE_LITERAL("void EmptyFunction(){}"));
+			TEST_EQUAL(TEXT("MinifyShader: EmptyFunction: C99 pragma 1"), MinifiedParsed.Chunks[1].GetCode(), SHADER_SOURCE_LITERAL("_Pragma(\"\")"));
+			TEST_EQUAL(TEXT("MinifyShader: EmptyFunction: C99 pragma 2"), MinifiedParsed.Chunks[2].GetCode(), SHADER_SOURCE_LITERAL("_Pragma(\"dxc diagnostic push\")"));
+			TEST_EQUAL(TEXT("MinifyShader: EmptyFunction: C99 pragma 3"), MinifiedParsed.Chunks[3].GetCode(), SHADER_SOURCE_LITERAL("_Pragma(\"dxc diagnostic ignored \\\"-Wall\\\"\") _Pragma(\"dxc diagnostic ignored \\\"-Wconversion\\\"\")"));
+			TEST_EQUAL(TEXT("MinifyShader: EmptyFunction: C99 pragma 4"), MinifiedParsed.Chunks[4].GetCode(), SHADER_SOURCE_LITERAL("_Pragma(\"dxc diagnostic pop\")"));
+			TEST_EQUAL(TEXT("MinifyShader: EmptyFunction: define"), MinifiedParsed.Chunks[5].GetCode(), SHADER_SOURCE_LITERAL("#define COMPILER_DEFINITION_TEST 123"));
+			TEST_EQUAL(TEXT("MinifyShader: EmptyFunction: function"), MinifiedParsed.Chunks[6].GetCode(), SHADER_SOURCE_LITERAL("void EmptyFunction(){}"));
 		}
 	}
 
