[ShaderSymbols] Fixed race condition in FShaderSymbolExport::NotifyShaderCompilersShutdown().

This function was missing a lock that the other notify function NotifyShaderCompiled() was using.

#rnx

diff --git a/Engine/Source/Developer/ShaderCompilerCommon/Private/ShaderSymbolExport.cpp b/Engine/Source/Developer/ShaderCompilerCommon/Private/ShaderSymbolExport.cpp
--- a/Engine/Source/Developer/ShaderCompilerCommon/Private/ShaderSymbolExport.cpp
+++ b/Engine/Source/Developer/ShaderCompilerCommon/Private/ShaderSymbolExport.cpp
@@ -240,6 +240,7 @@ void FShaderSymbolExport::WriteSymbolData(const FString& Filename, const FString
 
 void FShaderSymbolExport::NotifyShaderCompilersShutdown()
 {
+	FScopeLock Lock(&Cs);
 	if (ShaderInfos.Num())
 	{
 		if (InfoFilePath.Len())
