Shader proprocessing parser extended. Added support for custom namespaces, global arrays of textures / buffers and free form ',' expressions.

diff --git a/Engine/Source/Developer/ShaderCompilerCommon/Private/HlslLexer.cpp b/Engine/Source/Developer/ShaderCompilerCommon/Private/HlslLexer.cpp
--- a/Engine/Source/Developer/ShaderCompilerCommon/Private/HlslLexer.cpp
+++ b/Engine/Source/Developer/ShaderCompilerCommon/Private/HlslLexer.cpp
@@ -461,6 +461,7 @@ namespace CrossCompiler
 				InsertToken(TEXT("inline"), EHlslToken::Inline);
 				InsertToken(TEXT("typedef"), EHlslToken::Typedef);
 				InsertToken(TEXT("packoffset"), EHlslToken::PackOffset);
+				InsertToken(TEXT("namespace"), EHlslToken::Namespace);
 				InsertToken(TEXT("operator"), EHlslToken::Operator);
 				InsertToken(TEXT("_Static_assert"), EHlslToken::StaticAssert); // HLSL2021 adopted C11 '_Static_assert'-statements
 				InsertToken(TEXT("static_assert"), EHlslToken::StaticAssert); // Some shader compilers support C++11 'static_assert'-statements
diff --git a/Engine/Source/Developer/ShaderCompilerCommon/Private/HlslLexer.h b/Engine/Source/Developer/ShaderCompilerCommon/Private/HlslLexer.h
--- a/Engine/Source/Developer/ShaderCompilerCommon/Private/HlslLexer.h
+++ b/Engine/Source/Developer/ShaderCompilerCommon/Private/HlslLexer.h
@@ -298,6 +298,7 @@ namespace CrossCompiler
 		Inline,
 		Typedef,
 		PackOffset,
+		Namespace,
 
 		Identifier,
 		Operator, // HLSL2021
diff --git a/Engine/Source/Developer/ShaderCompilerCommon/Private/HlslParser.cpp b/Engine/Source/Developer/ShaderCompilerCommon/Private/HlslParser.cpp
--- a/Engine/Source/Developer/ShaderCompilerCommon/Private/HlslParser.cpp
+++ b/Engine/Source/Developer/ShaderCompilerCommon/Private/HlslParser.cpp
@@ -14,6 +14,7 @@ namespace CrossCompiler
 	EParseResult ParseStructDeclaration(FHlslParser& Parser, FSymbolScope* SymbolScope, FLinearAllocator* Allocator, AST::FTypeSpecifier** OutTypeSpecifier);
 	EParseResult TryParseAttribute(FHlslParser& Parser, FLinearAllocator* Allocator, AST::FAttribute** OutAttribute);
 	EParseResult ParseFunctionDeclaration(FHlslParser& Parser, FLinearAllocator* Allocator, TLinearArray<AST::FAttribute*>& Attributes, AST::FNode** OutFunction);
+	EParseResult TryTranslationUnit(FHlslParser& Parser, FLinearAllocator* Allocator, AST::FNode** OutNode);
 
 	typedef EParseResult(*TTryRule)(class FHlslParser& Scanner, FLinearAllocator* Allocator, AST::FNode** OutStatement);
 	struct FRulePair
@@ -341,7 +342,8 @@ namespace CrossCompiler
 				{
 					if (Scanner.MatchToken(EHlslToken::Equal))
 					{
-						const bool bAllowInitializerList = false;
+						// Setting to true allows declaration of global arrays of textures / buffers
+						const bool bAllowInitializerList = true;
 						if (ParseInitializer(Scanner, SymbolScope, bAllowInitializerList, Allocator, &Declaration->Initializer) != EParseResult::Matched)
 						{
 							Scanner.SourceError(TEXT("Invalid initializer\n"));
@@ -686,6 +688,13 @@ namespace CrossCompiler
 
 						DeclaratorList->Declarations.Add(Declaration);
 						*OutDeclaratorList = DeclaratorList;
+
+						// Adding variable declarations to symbols defined within specific namespace
+						if (SymbolScope && SymbolScope->Name != nullptr)
+						{
+							SymbolScope->Symbols.Add(IdentifierToken->String);
+						}
+
 						return EParseResult::Matched;
 					}
 				}
@@ -1302,6 +1311,38 @@ Done:
 		return TryStatementRules(Parser, Allocator, OutStatement);
 	}
 
+	EParseResult ParseNamespaceBlock(FHlslParser& Parser, FLinearAllocator* Allocator, AST::FNode** OutStatement)
+	{
+		const FHlslToken* Token = Parser.Scanner.PeekToken();
+		const FHlslToken* Token1 = Parser.Scanner.PeekToken(1);
+		const FHlslToken* Token2 = Parser.Scanner.PeekToken(2);
+
+		if (Token && Token->Token == EHlslToken::Namespace && 
+			Token1 && Token1->Token == EHlslToken::Identifier && 
+			Token2 && Token2->Token == EHlslToken::LeftBrace)
+		{
+			FCreateSymbolScope NamespaceSymbolScope(Allocator, &Parser.CurrentScope);
+			Parser.CurrentScope->Name = *(Token1->String);
+
+			check(Parser.Scanner.MatchToken(EHlslToken::Namespace));
+			check(Parser.Scanner.MatchToken(EHlslToken::Identifier));
+			check(Parser.Scanner.MatchToken(EHlslToken::LeftBrace));
+
+			do
+			{
+				auto Result = TryTranslationUnit(Parser, Allocator, OutStatement);
+				if (Result != EParseResult::Matched)
+				{
+					return ParseResultError();
+				}
+			}while(!Parser.Scanner.MatchToken(EHlslToken::RightBrace));
+
+			return EParseResult::Matched;
+		}
+
+		return EParseResult::NotMatched;
+	}
+	
 	EParseResult ParseStatementBlock(FHlslParser& Parser, FLinearAllocator* Allocator, AST::FNode** OutStatement)
 	{
 		FCreateSymbolScope SymbolScope(Allocator, &Parser.CurrentScope);
@@ -1443,6 +1484,12 @@ Done:
 				FunctionDefinition->Body = (AST::FCompoundStatement*)Body;
 				FunctionDefinition->Prototype = Function;
 				*OutFunction = FunctionDefinition;
+				
+				// Adding function declarations to symbols defined within specific namespace
+				if (Parser.CurrentScope && Parser.CurrentScope->Name != nullptr)
+				{
+					Parser.CurrentScope->Symbols.Add(FunctionDefinition->Prototype->Identifier);
+				}
 			}
 		}
 
@@ -1451,6 +1498,19 @@ Done:
 
 	EParseResult ParseLocalDeclaration(FHlslParser& Parser, FLinearAllocator* Allocator, AST::FNode** OutDeclaration)
 	{
+		// We don't support namespace declarations in Local scopes, only in Global scopes. Local declarations that start with <namespace>:: are not allowed in language syntax
+		const FHlslToken* Token = Parser.Scanner.PeekToken();
+		const FHlslToken* Token2 = Parser.Scanner.PeekToken(1);
+		if (Token && !Token->String.IsEmpty() && Token2 && Token2->Token == EHlslToken::ColonColon)
+		{
+			const FSymbolScope* Namespace = Parser.CurrentScope->GetGlobalScope();
+			Namespace = Namespace->FindNamespace(*Token->String);
+			if (Namespace)
+			{
+				return EParseResult::NotMatched;
+			}
+		}
+
 		AST::FDeclaratorList* List = nullptr;
 		auto Result = ParseGeneralDeclaration(
 			Parser, Parser.CurrentScope, Allocator, &List, 0,
@@ -2079,7 +2139,8 @@ Done:
 		auto* Statement = new(Allocator) AST::FExpressionStatement(Allocator, nullptr, Parser.Scanner.GetCurrentToken()->SourceInfo);
 		if (ParseExpression(Parser.Scanner, Parser.CurrentScope, EEF_ALLOW_ASSIGNMENT, Allocator, &Statement->Expression) == EParseResult::Matched)
 		{
-			if (Parser.Scanner.MatchToken(EHlslToken::Semicolon))
+			// Adding Comma to handle free form expressions that can sometimes be generated from blueprints. We could adjust RulesStatements and swap ParseExpressionStatement there with ParseExpressionList2 but this seems to work as well and less code changes
+			if (Parser.Scanner.MatchToken(EHlslToken::Semicolon) || Parser.Scanner.MatchToken(EHlslToken::Comma))
 			{
 				*OutStatement = Statement;
 				return EParseResult::Matched;
@@ -2267,6 +2328,12 @@ Done:
 			return Result;
 		}
 
+		Result = ParseNamespaceBlock(Parser, Allocator, OutNode);
+		if (Result == EParseResult::Error || Result == EParseResult::Matched)
+		{
+			return Result;
+		}
+
 		// Ignore semicolons with no code before them.
 		if (Parser.Scanner.MatchToken(EHlslToken::Semicolon))
 		{
diff --git a/Engine/Source/Developer/ShaderCompilerCommon/Private/HlslUtils.cpp b/Engine/Source/Developer/ShaderCompilerCommon/Private/HlslUtils.cpp
--- a/Engine/Source/Developer/ShaderCompilerCommon/Private/HlslUtils.cpp
+++ b/Engine/Source/Developer/ShaderCompilerCommon/Private/HlslUtils.cpp
@@ -239,7 +239,10 @@ struct FRemoveAlgorithm
 		return NewReturnType;
 	}
 
-	CrossCompiler::AST::FFunctionDefinition* CreateNewEntryFunction(CrossCompiler::AST::FCompoundStatement* Body, CrossCompiler::AST::FFullySpecifiedType* ReturnType, TArray<CrossCompiler::AST::FParameterDeclarator*>& Parameters)
+	CrossCompiler::AST::FFunctionDefinition* CreateNewEntryFunction(CrossCompiler::AST::FCompoundStatement* Body, 
+		CrossCompiler::AST::FFullySpecifiedType* ReturnType, 
+		TArray<CrossCompiler::AST::FParameterDeclarator*>& Parameters, 
+		CrossCompiler::TLinearArray<CrossCompiler::AST::FAttribute*>& FunctionAttributes)
 	{
 		using namespace CrossCompiler::AST;
 		// New Entry definition/prototype
@@ -247,6 +250,7 @@ struct FRemoveAlgorithm
 		NewEntryFunction->Prototype = new(Allocator) FFunction(Allocator, SourceInfo);
 		NewEntryFunction->Prototype->Identifier = Allocator->Strdup(*(EntryPoint + TEXT("__OPTIMIZED")));
 		NewEntryFunction->Prototype->ReturnType = ReturnType;
+		NewEntryFunction->Prototype->Attributes = FunctionAttributes;
 		NewEntryFunction->Body = Body;
 		for (auto* Parameter : Parameters)
 		{
@@ -553,7 +557,7 @@ struct FRemoveUnusedOutputs : FRemoveAlgorithm
 		auto* CallInstruction = new(Allocator) CrossCompiler::AST::FExpressionStatement(Allocator, BodyContext.CallExpression, SourceInfo);
 
 		FCompoundStatement* Body = AddStatementsToBody(BodyContext, CallInstruction);
-		FFunctionDefinition* NewEntryFunction = CreateNewEntryFunction(Body, MakeSimpleType(BodyContext.NewReturnStruct->Name), BodyContext.NewFunctionParameters);
+		FFunctionDefinition* NewEntryFunction = CreateNewEntryFunction(Body, MakeSimpleType(BodyContext.NewReturnStruct->Name), BodyContext.NewFunctionParameters, EntryFunction->Prototype->Attributes);
 		EntryPoint = NewEntryFunction->Prototype->Identifier;
 		WriteGeneratedOutCode(NewEntryFunction, BodyContext.NewStructs, GeneratedCode);
 		bSuccess = true;
@@ -1033,7 +1037,7 @@ struct FRemoveUnusedInputs : FRemoveAlgorithm
 			BodyContext.NewFunctionParameters.Add(Declarator);
 		}
 
-		FFunctionDefinition* NewEntryFunction = CreateNewEntryFunction(Body, EntryFunction->Prototype->ReturnType, BodyContext.NewFunctionParameters);
+		FFunctionDefinition* NewEntryFunction = CreateNewEntryFunction(Body, EntryFunction->Prototype->ReturnType, BodyContext.NewFunctionParameters, EntryFunction->Prototype->Attributes);
 		NewEntryFunction->Prototype->ReturnSemantic = EntryFunction->Prototype->ReturnSemantic;
 
 		WriteGeneratedInCode(NewEntryFunction, BodyContext.NewStructs, GeneratedCode);
@@ -1619,6 +1623,7 @@ bool FindEntryPointParameters(
 	const FString& InSourceCode,
 	const FString& InEntryPoint,
 	EShaderParameterStorageClass ParameterStorageClass,
+	TConstArrayView<FScopedDeclarations> InScopedDeclarations,
 	TArray<FString>& OutParameterSemantics,
 	TArray<FString>& OutErrors)
 {
@@ -1632,7 +1637,7 @@ bool FindEntryPointParameters(
 			Data.FindEntryPointParameters(ASTNodes);
 		};
 	CrossCompiler::FCompilerMessages Messages;
-	if (!CrossCompiler::Parser::Parse(InSourceCode, DummyFilename, Messages, {}, ResultCallbackFunction))
+	if (!CrossCompiler::Parser::Parse(InSourceCode, DummyFilename, Messages, InScopedDeclarations, ResultCallbackFunction))
 	{
 		Data.Errors.Add(FString(TEXT("FindEntryPointParameters: Failed to parse HLSL source!")));
 		OutErrors = MoveTemp(Data.Errors);
diff --git a/Engine/Source/Developer/ShaderCompilerCommon/Public/ShaderCompilerCommon.h b/Engine/Source/Developer/ShaderCompilerCommon/Public/ShaderCompilerCommon.h
--- a/Engine/Source/Developer/ShaderCompilerCommon/Public/ShaderCompilerCommon.h
+++ b/Engine/Source/Developer/ShaderCompilerCommon/Public/ShaderCompilerCommon.h
@@ -414,6 +414,7 @@ extern SHADERCOMPILERCOMMON_API bool FindEntryPointParameters(
 	const FString& InSourceCode,
 	const FString& InEntryPoint,
 	EShaderParameterStorageClass ParameterStorageClass,
+	TConstArrayView<FScopedDeclarations> InScopedDeclarations,
 	TArray<FString>& OutParameterSemantics,
 	TArray<FString>& OutErrors
 );
diff --git a/Engine/Source/Developer/Windows/ShaderFormatD3D/Private/D3DShaderCompiler.cpp b/Engine/Source/Developer/Windows/ShaderFormatD3D/Private/D3DShaderCompiler.cpp
--- a/Engine/Source/Developer/Windows/ShaderFormatD3D/Private/D3DShaderCompiler.cpp
+++ b/Engine/Source/Developer/Windows/ShaderFormatD3D/Private/D3DShaderCompiler.cpp
@@ -488,7 +488,7 @@ static bool PatchHlslWithReorderedIOVariables(
 
 	// Parse variable names from SPIR-V input
 	TArray<FString> Variables, ParsingErrors;
-	if (!FindEntryPointParameters(OriginalShaderSource, OriginalEntryPoint, StageVariablesStorageClass, Variables, ParsingErrors))
+	if (!FindEntryPointParameters(OriginalShaderSource, OriginalEntryPoint, StageVariablesStorageClass, {}, Variables, ParsingErrors))
 	{
 		for (FString& Error : ParsingErrors)
 		{
