Deprecate SHADER_SOURCE_ANSI (and assume it is always 1) and related macros, typedefs and helper functions

* Remove code that assumes SHADER_SOURCE_ANSI=0

diff --git a/Engine/Source/Developer/ShaderCompilerCommon/Private/ShaderCompilerCommon.cpp b/Engine/Source/Developer/ShaderCompilerCommon/Private/ShaderCompilerCommon.cpp
--- a/Engine/Source/Developer/ShaderCompilerCommon/Private/ShaderCompilerCommon.cpp
+++ b/Engine/Source/Developer/ShaderCompilerCommon/Private/ShaderCompilerCommon.cpp
@@ -516,7 +516,6 @@ bool UE::ShaderCompilerCommon::RemoveDeadCode(FShaderSource& InOutPreprocessedSh
 	           |  UE::ShaderMinifier::EMinifyShaderFlags::OutputStats;  // Output a comment detailing how many blocks of each type (functions/structs/etc.) were emitted
 #endif
 
-#if SHADER_SOURCE_ANSI
 	TArray<FShaderSource::FStringType> ConvertedRequiredSymbols;
 	TArray<FShaderSource::FViewType> RequiredSymbolViews;
 	for (FStringView InSymbol : InRequiredSymbols)
@@ -525,9 +524,6 @@ bool UE::ShaderCompilerCommon::RemoveDeadCode(FShaderSource& InOutPreprocessedSh
 		ConvertedString.Append(InSymbol);
 		RequiredSymbolViews.Add(FShaderSource::FViewType(ConvertedString));
 	}
-#else
-	TConstArrayView<FStringView> RequiredSymbolViews = InRequiredSymbols;
-#endif
 
 	UE::ShaderMinifier::FMinifiedShader Minified  = UE::ShaderMinifier::Minify(InOutPreprocessedShaderSource, RequiredSymbolViews,
 		  UE::ShaderMinifier::EMinifyShaderFlags::OutputCommentLines // Preserve comments that were left after preprocessing
@@ -1332,7 +1328,7 @@ void CleanupUniformBufferCode(const FShaderCompilerEnvironment& Environment, FSh
 	UniformBufferSpans.Reserve(Environment.UniformBufferMap.Num());
 	memset(UniformBuffersByLength, 0xff, sizeof(UniformBuffersByLength));
 
-	FShaderSource::FViewType UniformBufferStructIdentifier = SHADER_SOURCE_VIEWLITERAL("UniformBuffer");
+	FShaderSource::FViewType UniformBufferStructIdentifier = ANSITEXTVIEW("UniformBuffer");
 
 	FShaderSource::CharType* SourceStart = PreprocessedShaderSource.GetData();
 	FShaderSource::CharType* SearchPtr = SourceStart;
diff --git a/Engine/Source/Developer/ShaderCompilerCommon/Private/ShaderMinifier.cpp b/Engine/Source/Developer/ShaderCompilerCommon/Private/ShaderMinifier.cpp
--- a/Engine/Source/Developer/ShaderCompilerCommon/Private/ShaderMinifier.cpp
+++ b/Engine/Source/Developer/ShaderCompilerCommon/Private/ShaderMinifier.cpp
@@ -189,47 +189,47 @@ static FShaderSource::FViewType ExtractOperator(FShaderSource::FViewType Source)
 	static const FShaderSource::FViewType SupportedOperators[] =
 	{
 		// three-character operators
-		SHADER_SOURCE_VIEWLITERAL("<<="),
-		SHADER_SOURCE_VIEWLITERAL(">>="),
-		SHADER_SOURCE_VIEWLITERAL("->*"),
+		ANSITEXTVIEW("<<="),
+		ANSITEXTVIEW(">>="),
+		ANSITEXTVIEW("->*"),
 
 		// two-character operators
-		SHADER_SOURCE_VIEWLITERAL("+="),
-		SHADER_SOURCE_VIEWLITERAL("++"),
-		SHADER_SOURCE_VIEWLITERAL("-="),
-		SHADER_SOURCE_VIEWLITERAL("--"),
-		SHADER_SOURCE_VIEWLITERAL("->"),
-		SHADER_SOURCE_VIEWLITERAL("*="),
-		SHADER_SOURCE_VIEWLITERAL("/="),
-		SHADER_SOURCE_VIEWLITERAL("%="),
-		SHADER_SOURCE_VIEWLITERAL("^="),
-		SHADER_SOURCE_VIEWLITERAL("&="),
-		SHADER_SOURCE_VIEWLITERAL("&&"),
-		SHADER_SOURCE_VIEWLITERAL("|="),
-		SHADER_SOURCE_VIEWLITERAL("||"),
-		SHADER_SOURCE_VIEWLITERAL("<<"),
-		SHADER_SOURCE_VIEWLITERAL("<="),
-		SHADER_SOURCE_VIEWLITERAL(">>"),
-		SHADER_SOURCE_VIEWLITERAL(">="),
-		SHADER_SOURCE_VIEWLITERAL("=="),
-		SHADER_SOURCE_VIEWLITERAL("!="),
-		SHADER_SOURCE_VIEWLITERAL("()"),
-		SHADER_SOURCE_VIEWLITERAL("[]"),
+		ANSITEXTVIEW("+="),
+		ANSITEXTVIEW("++"),
+		ANSITEXTVIEW("-="),
+		ANSITEXTVIEW("--"),
+		ANSITEXTVIEW("->"),
+		ANSITEXTVIEW("*="),
+		ANSITEXTVIEW("/="),
+		ANSITEXTVIEW("%="),
+		ANSITEXTVIEW("^="),
+		ANSITEXTVIEW("&="),
+		ANSITEXTVIEW("&&"),
+		ANSITEXTVIEW("|="),
+		ANSITEXTVIEW("||"),
+		ANSITEXTVIEW("<<"),
+		ANSITEXTVIEW("<="),
+		ANSITEXTVIEW(">>"),
+		ANSITEXTVIEW(">="),
+		ANSITEXTVIEW("=="),
+		ANSITEXTVIEW("!="),
+		ANSITEXTVIEW("()"),
+		ANSITEXTVIEW("[]"),
 
 		// single character operators
-		SHADER_SOURCE_VIEWLITERAL("+"),
-		SHADER_SOURCE_VIEWLITERAL("-"),
-		SHADER_SOURCE_VIEWLITERAL("*"),
-		SHADER_SOURCE_VIEWLITERAL("/"),
-		SHADER_SOURCE_VIEWLITERAL("%"),
-		SHADER_SOURCE_VIEWLITERAL("^"),
-		SHADER_SOURCE_VIEWLITERAL("&"),
-		SHADER_SOURCE_VIEWLITERAL("|"),
-		SHADER_SOURCE_VIEWLITERAL("~"),
-		SHADER_SOURCE_VIEWLITERAL("!"),
-		SHADER_SOURCE_VIEWLITERAL("="),
-		SHADER_SOURCE_VIEWLITERAL("<"),
-		SHADER_SOURCE_VIEWLITERAL(">"),
+		ANSITEXTVIEW("+"),
+		ANSITEXTVIEW("-"),
+		ANSITEXTVIEW("*"),
+		ANSITEXTVIEW("/"),
+		ANSITEXTVIEW("%"),
+		ANSITEXTVIEW("^"),
+		ANSITEXTVIEW("&"),
+		ANSITEXTVIEW("|"),
+		ANSITEXTVIEW("~"),
+		ANSITEXTVIEW("!"),
+		ANSITEXTVIEW("="),
+		ANSITEXTVIEW("<"),
+		ANSITEXTVIEW(">"),
 	};
 
 	for (FShaderSource::FViewType Operator : SupportedOperators)
@@ -248,7 +248,7 @@ static FShaderSource::FViewType ExtractOperator(FShaderSource::FViewType Source)
 template <int CompareType>
 inline int32 ScanPastCharactersSimd(const FShaderSource::CharType* Source, int32 SourceLen, __m128i NeedleVec)
 {
-	constexpr int32 Mode = (FShaderSource::IsWide()  ? _SIDD_UWORD_OPS : _SIDD_UBYTE_OPS) | CompareType | _SIDD_MASKED_NEGATIVE_POLARITY;
+	constexpr int32 Mode = _SIDD_UBYTE_OPS | CompareType | _SIDD_MASKED_NEGATIVE_POLARITY;
 	int32 Cursor = 0;
 	while (Cursor < SourceLen)
 	{
@@ -259,7 +259,7 @@ inline int32 ScanPastCharactersSimd(const FShaderSource::CharType* Source, int32
 			Cursor += _mm_cmpistri(NeedleVec, Chunk, Mode);
 			break;
 		}
-		Cursor += FShaderSource::GetSimdCharCount();
+		Cursor += 16;
 	}
 	return Cursor;
 }
@@ -272,10 +272,8 @@ static FShaderSource::FViewType SkipUntilNonIdentifierCharacter(FShaderSource::F
 	const FShaderSource::CharType* SourceData = Source.GetData();
 
 #if UE_SHADER_MINIFIER_SSE
-	const __m128i NeedleVec = FShaderSource::IsWide()
-		? _mm_setr_epi16(L'0', L'9', L'a', L'z', L'A', L'Z', L'_', L'_')
-		: _mm_setr_epi8('0', '9', 'a', 'z', 'A', 'Z', '_', '_', 0, 0, 0, 0, 0, 0, 0, 0);
-		
+	const __m128i NeedleVec = _mm_setr_epi8('0', '9', 'a', 'z', 'A', 'Z', '_', '_', 0, 0, 0, 0, 0, 0, 0, 0);
+
 	Cursor = ScanPastCharactersSimd<_SIDD_CMP_RANGES>(SourceData, SourceLen, NeedleVec);
 #else
 	while (Cursor < SourceLen)
@@ -290,7 +288,7 @@ static FShaderSource::FViewType SkipUntilNonIdentifierCharacter(FShaderSource::F
 	return FShaderSource::FViewType(SourceData + Cursor, FMath::Max(SourceLen - Cursor, 0));
 }
 
-static FShaderSource::FViewType SkipUntilNonNumber(FShaderSource::FViewType Source) 
+static FShaderSource::FViewType SkipUntilNonNumber(FShaderSource::FViewType Source)
 {
 	const int32 SourceLen = Source.Len();
 	int32 Cursor = 0;
@@ -313,9 +311,7 @@ static FShaderSource::FViewType SkipSpace(FShaderSource::FViewType Source)
 	const FShaderSource::CharType* SourceData = Source.GetData();
 
 #if UE_SHADER_MINIFIER_SSE
-	const __m128i NeedleVec = FShaderSource::IsWide()
-		? _mm_setr_epi16(L' ', L'\f', L'\r', L'\n', L'\t', L'\v', 0, 0)
-		: _mm_setr_epi8(' ', '\f', '\r', '\n', '\t', '\v', 0, 0, 0, 0, 0, 0, 0, 0, 0, 0);
+	const __m128i NeedleVec = _mm_setr_epi8(' ', '\f', '\r', '\n', '\t', '\v', 0, 0, 0, 0, 0, 0, 0, 0, 0, 0);
 
 	Cursor = ScanPastCharactersSimd<_SIDD_CMP_EQUAL_ANY>(SourceData, SourceLen, NeedleVec);
 #else
@@ -424,7 +420,7 @@ static FShaderSource::FViewType ExtractBlock(FShaderSource::FViewType Source, FS
 		else if (C == '#')
 		{
 			FShaderSource::FViewType Source(SourceData + Cursor, SourceLen - Cursor);
-			if (StartsWith(Source, SHADER_SOURCE_VIEWLITERAL("#line")))
+			if (StartsWith(Source, ANSITEXTVIEW("#line")))
 			{
 				FShaderSource::FViewType Remainder = SkipUntilNextLine(Source);
 				FShaderSource::FViewType Block = SubStrView(Source, 0, Source.Len() - Remainder.Len());
@@ -436,14 +432,12 @@ static FShaderSource::FViewType ExtractBlock(FShaderSource::FViewType Source, FS
 	};
 
 #if UE_SHADER_MINIFIER_SSE
-	const __m128i NeedleVec = FShaderSource::IsWide()
-		? _mm_setr_epi16(DelimBegin, DelimEnd, L'#', 0, 0, 0, 0, 0)
-		: _mm_setr_epi8(DelimBegin, DelimEnd, '#', 0, 0, 0, 0, 0, 0, 0, 0, 0, 0, 0, 0, 0);
+	const __m128i NeedleVec = _mm_setr_epi8(DelimBegin, DelimEnd, '#', 0, 0, 0, 0, 0, 0, 0, 0, 0, 0, 0, 0, 0);
 
 	while (Cursor < SourceLen && Status != EStatus::Finished)
 	{
 		__m128i Chunk = _mm_loadu_si128(reinterpret_cast<const __m128i*>(SourceData + Cursor));
-		constexpr int32 Mode = (FShaderSource::IsWide() ? _SIDD_UWORD_OPS : _SIDD_UBYTE_OPS) | _SIDD_CMP_EQUAL_ANY | _SIDD_MOST_SIGNIFICANT;
+		constexpr int32 Mode = _SIDD_UBYTE_OPS | _SIDD_CMP_EQUAL_ANY | _SIDD_MOST_SIGNIFICANT;
 		const int32 CompareResult = _mm_cmpistrc(NeedleVec, Chunk, Mode);
 		if (CompareResult)
 		{
@@ -454,10 +448,10 @@ static FShaderSource::FViewType ExtractBlock(FShaderSource::FViewType Source, FS
 				const uint32 BitIndex = FMath::CountTrailingZeros(Mask);
 				const uint32 ChunkCharIndex = BitIndex / sizeof(FShaderSource::CharType);
 				Status = ProcessCharacter(Cursor + ChunkCharIndex);
-				Mask &= ~(FShaderSource::GetSingleCharMask() << BitIndex);
+				Mask &= ~(1u << BitIndex);
 			}
 		}
-		Cursor += FShaderSource::GetSimdCharCount();
+		Cursor += 16;
 	}
 #else
 	while (Cursor < SourceLen && Status != EStatus::Finished)
@@ -553,7 +547,7 @@ struct FNamespace
 			for (const FShaderSource::FViewType& Part : InStack)
 			{
 				FullName += Part;
-				FullName += SHADER_SOURCE_LITERAL("::");
+				FullName += "::";
 			}
 			FullName.LeftChopInline(2);
 		}
@@ -901,7 +895,7 @@ static FParsedShader ParseShader(const FShaderSource& InSource, FDiagnostics& Ou
 		{
 			if (FirstChar == '/')
 			{
-				if (StartsWith(Source, SHADER_SOURCE_VIEWLITERAL("//")))
+				if (StartsWith(Source, ANSITEXTVIEW("//")))
 				{
 					FShaderSource::FViewType Remainder = SkipUntilNextLine(Source);
 
@@ -918,9 +912,9 @@ static FParsedShader ParseShader(const FShaderSource& InSource, FDiagnostics& Ou
 
 					continue;
 				}
-				else if (StartsWith(Source, SHADER_SOURCE_VIEWLITERAL("/*")))
+				else if (StartsWith(Source, ANSITEXTVIEW("/*")))
 				{
-					Source = SkipUntilStr(Source, SHADER_SOURCE_VIEWLITERAL("*/"));
+					Source = SkipUntilStr(Source, ANSITEXTVIEW("*/"));
 					if (Source.Len() >= 2)
 					{
 						Source = SubStrView(Source, 2);
@@ -930,7 +924,7 @@ static FParsedShader ParseShader(const FShaderSource& InSource, FDiagnostics& Ou
 			}
 			else if (FirstChar == '#')
 			{
-				if (StartsWith(Source, SHADER_SOURCE_VIEWLITERAL("#line")))
+				if (StartsWith(Source, ANSITEXTVIEW("#line")))
 				{
 					FShaderSource::FViewType Remainder = SkipUntilNextLine(Source);
 					FShaderSource::FViewType Block = SubStrView(Source, 0, Source.Len() - Remainder.Len());
@@ -938,7 +932,7 @@ static FParsedShader ParseShader(const FShaderSource& InSource, FDiagnostics& Ou
 					Source = Remainder;
 					continue;
 				}
-				else if (StartsWith(Source, SHADER_SOURCE_VIEWLITERAL("#pragma")))
+				else if (StartsWith(Source, ANSITEXTVIEW("#pragma")))
 				{
 					FShaderSource::FViewType Remainder = SkipUntilNextLine(Source);
 					FShaderSource::FViewType Block = SubStrView(Source, 0, Source.Len() - Remainder.Len());
@@ -948,7 +942,7 @@ static FParsedShader ParseShader(const FShaderSource& InSource, FDiagnostics& Ou
 					Source = Remainder;
 					continue;
 				}
-				else if (StartsWith(Source, SHADER_SOURCE_VIEWLITERAL("#define")))
+				else if (StartsWith(Source, ANSITEXTVIEW("#define")))
 				{
 					// TODO: handle `\` new lines in defines
 					FShaderSource::FViewType Remainder = SkipUntilNextLine(Source);
@@ -959,9 +953,9 @@ static FParsedShader ParseShader(const FShaderSource& InSource, FDiagnostics& Ou
 					Source = Remainder;
 					continue;
 				}
-				else if (StartsWith(Source, SHADER_SOURCE_VIEWLITERAL("#if 0")))
+				else if (StartsWith(Source, ANSITEXTVIEW("#if 0")))
 				{
-					Source = SkipUntilStr(Source, SHADER_SOURCE_VIEWLITERAL("#endif"));
+					Source = SkipUntilStr(Source, ANSITEXTVIEW("#endif"));
 					if (Source.Len() >= 6)
 					{
 						Source = SubStrView(Source, 6);
@@ -969,7 +963,7 @@ static FParsedShader ParseShader(const FShaderSource& InSource, FDiagnostics& Ou
 					continue;
 				}
 			}
-			else if (FirstChar == '_' && StartsWith(Source, SHADER_SOURCE_VIEWLITERAL("_Pragma")))
+			else if (FirstChar == '_' && StartsWith(Source, ANSITEXTVIEW("_Pragma")))
 			{
 				FShaderSource::FViewType PragmaStart = Source;
 				FShaderSource::FViewType Remainder = Source;
@@ -977,7 +971,7 @@ static FParsedShader ParseShader(const FShaderSource& InSource, FDiagnostics& Ou
 				// handle case of multiple _Pragma operators on the same line (to avoid inserting newlines)
 				while (PragmaStart.Len() > 0)
 				{
-					Remainder = SkipUntilStr(Remainder, SHADER_SOURCE_VIEWLITERAL(")"));
+					Remainder = SkipUntilStr(Remainder, ANSITEXTVIEW(")"));
 					if (Remainder.Len() >= 1)
 					{
 						Remainder = SubStrView(Remainder, 1);
@@ -985,10 +979,10 @@ static FParsedShader ParseShader(const FShaderSource& InSource, FDiagnostics& Ou
 
 					FShaderSource::FViewType NextLine = SkipUntilNextLine(Remainder);
 					FShaderSource::FViewType LineRemainder = SubStrView(Remainder, 0, Remainder.Len() - NextLine.Len());
-					PragmaStart = SkipUntilStr(LineRemainder, SHADER_SOURCE_VIEWLITERAL("_Pragma"));
+					PragmaStart = SkipUntilStr(LineRemainder, ANSITEXTVIEW("_Pragma"));
 					if (PragmaStart.Len() > 0)
 					{
-						Remainder = SkipUntilStr(Remainder, SHADER_SOURCE_VIEWLITERAL("_Pragma"));
+						Remainder = SkipUntilStr(Remainder, ANSITEXTVIEW("_Pragma"));
 					}
 				}
 
@@ -1001,7 +995,7 @@ static FParsedShader ParseShader(const FShaderSource& InSource, FDiagnostics& Ou
 			}
 			else if (PendingBlocks.IsEmpty() && (FirstChar == '{' || FirstChar == '}'))
 			{
-				if (StartsWith(Source, SHADER_SOURCE_VIEWLITERAL("{")))
+				if (StartsWith(Source, ANSITEXTVIEW("{")))
 				{
 					if (ChunkType == ECodeChunkType::Namespace)
 					{
@@ -1025,7 +1019,7 @@ static FParsedShader ParseShader(const FShaderSource& InSource, FDiagnostics& Ou
 					}
 					continue;
 				}
-				else if (StartsWith(Source, SHADER_SOURCE_VIEWLITERAL("}")))
+				else if (StartsWith(Source, ANSITEXTVIEW("}")))
 				{
 					if (NamespaceTracker.Pop())
 					{
@@ -1068,48 +1062,48 @@ static FParsedShader ParseShader(const FShaderSource& InSource, FDiagnostics& Ou
 		{
 			if (ChunkType == ECodeChunkType::Unknown)
 			{
-				if (Equals(Identifier, SHADER_SOURCE_VIEWLITERAL("struct")))
+				if (Equals(Identifier, ANSITEXTVIEW("struct")))
 				{
 					ChunkType = ECodeChunkType::Struct;
 					StructBlockIndex = PendingBlocks.Num();
 				}
-				else if (Equals(Identifier, SHADER_SOURCE_VIEWLITERAL("cbuffer")) || Equals(Identifier, SHADER_SOURCE_VIEWLITERAL("ConstantBuffer")))
+				else if (Equals(Identifier, ANSITEXTVIEW("cbuffer")) || Equals(Identifier, ANSITEXTVIEW("ConstantBuffer")))
 				{
 					ChunkType = ECodeChunkType::CBuffer;
 					CbufferBlockIndex = PendingBlocks.Num();
 				}
-				else if (Equals(Identifier, SHADER_SOURCE_VIEWLITERAL("enum")))
+				else if (Equals(Identifier, ANSITEXTVIEW("enum")))
 				{
 					ChunkType = ECodeChunkType::Enum;
 					EnumBlockIndex = PendingBlocks.Num();
 				}
-				else if (Equals(Identifier, SHADER_SOURCE_VIEWLITERAL("namespace")))
+				else if (Equals(Identifier, ANSITEXTVIEW("namespace")))
 				{
 					ChunkType = ECodeChunkType::Namespace;
 					Source = Remainder;
 					continue;
 				}
-				else if (Equals(Identifier, SHADER_SOURCE_VIEWLITERAL("using")))
+				else if (Equals(Identifier, ANSITEXTVIEW("using")))
 				{
 					ChunkType = ECodeChunkType::Using;
 					Source = Remainder;
 					AddBlock(EBlockType::Keyword, Identifier);
 					continue;
 				}
-				else if (Equals(Identifier, SHADER_SOURCE_VIEWLITERAL("typedef")))
+				else if (Equals(Identifier, ANSITEXTVIEW("typedef")))
 				{
 					ChunkType = ECodeChunkType::Typedef;
 					Source = Remainder;
 					AddBlock(EBlockType::Keyword, Identifier);
 					continue;
 				}
-				else if (Equals(Identifier, SHADER_SOURCE_VIEWLITERAL("template")))
+				else if (Equals(Identifier, ANSITEXTVIEW("template")))
 				{
 					Source = Remainder;
 					AddBlock(EBlockType::Keyword, Identifier);
 					continue;
 				}
-				else if (Equals(Identifier, SHADER_SOURCE_VIEWLITERAL("operator")))
+				else if (Equals(Identifier, ANSITEXTVIEW("operator")))
 				{
 					ChunkType = ECodeChunkType::Operator;
 					Source = Remainder;
@@ -1155,12 +1149,12 @@ static FParsedShader ParseShader(const FShaderSource& InSource, FDiagnostics& Ou
 
 		EBlockType BlockType = EBlockType::Unknown;
 
-		if (StartsWith(Source, SHADER_SOURCE_VIEWLITERAL("==")))
+		if (StartsWith(Source, ANSITEXTVIEW("==")))
 		{
 			AddDiagnostic(Output.Errors, TEXTVIEW("Unexpected sequence '=='"));
 			break;
 		}
-		else if (StartsWith(Source, SHADER_SOURCE_VIEWLITERAL("::")))
+		else if (StartsWith(Source, ANSITEXTVIEW("::")))
 		{
 			Block = SubStrView(Source, 0, 2);
 			Source = SubStrView(Source, 2);
@@ -1242,7 +1236,7 @@ static FParsedShader ParseShader(const FShaderSource& InSource, FDiagnostics& Ou
 			{
 				BlockType = EBlockType::Subscript;
 			}
-			else if (Block.Contains(SHADER_SOURCE_LITERAL("NodeID")))
+			else if (Block.Contains(ANSITEXTVIEW("NodeID")))
 			{
 				BlockType = EBlockType::NodeId;
 			}
@@ -1367,16 +1361,16 @@ static void ExtractIdentifiers(FShaderSource::FViewType InSource, TArray<FShader
 
 		if (FirstChar == '#')
 		{
-			if (StartsWith(Source, SHADER_SOURCE_VIEWLITERAL("#line"))
-				|| StartsWith(Source, SHADER_SOURCE_VIEWLITERAL("#pragma")))
+			if (StartsWith(Source, ANSITEXTVIEW("#line"))
+				|| StartsWith(Source, ANSITEXTVIEW("#pragma")))
 			{
 				Source = SkipUntilNextLine(Source);
 				Source = SkipSpace(Source);
 				continue;
 			}
-			else if (StartsWith(Source, SHADER_SOURCE_VIEWLITERAL("#if 0")))
+			else if (StartsWith(Source, ANSITEXTVIEW("#if 0")))
 			{
-				Source = SkipUntilStr(Source, SHADER_SOURCE_VIEWLITERAL("#endif"));
+				Source = SkipUntilStr(Source, ANSITEXTVIEW("#endif"));
 				if (Source.Len() >= 6)
 				{
 					Source = SubStrView(Source, 6);
@@ -1387,15 +1381,15 @@ static void ExtractIdentifiers(FShaderSource::FViewType InSource, TArray<FShader
 		}
 		else if (FirstChar == '/')
 		{
-			if (StartsWith(Source, SHADER_SOURCE_VIEWLITERAL("//")))
+			if (StartsWith(Source, ANSITEXTVIEW("//")))
 			{
 				Source = SkipUntilNextLine(Source);
 				Source = SkipSpace(Source);
 				continue;
 			}
-			else if (StartsWith(Source, SHADER_SOURCE_VIEWLITERAL("/*")))
+			else if (StartsWith(Source, ANSITEXTVIEW("/*")))
 			{
-				Source = SkipUntilStr(Source, SHADER_SOURCE_VIEWLITERAL("*/"));
+				Source = SkipUntilStr(Source, ANSITEXTVIEW("*/"));
 				if (Source.Len() >= 2)
 				{
 					Source = SubStrView(Source, 2);
@@ -1460,7 +1454,7 @@ static void OutputChunk(const FCodeChunk& Chunk, FShaderSource::FStringType& Out
 
 			if (Block.Type == EBlockType::Expression)
 			{
-				OutputStream.Append(SHADER_SOURCE_VIEWLITERAL("= "));
+				OutputStream.Append(ANSITEXTVIEW("= "));
 			}
 			else if (Block.Type == EBlockType::Body)
 			{
@@ -1469,7 +1463,7 @@ static void OutputChunk(const FCodeChunk& Chunk, FShaderSource::FStringType& Out
 
 			if (Block.Type == EBlockType::Binding || Block.Type == EBlockType::Base)
 			{
-				OutputStream.Append(SHADER_SOURCE_VIEWLITERAL(": "));
+				OutputStream.Append(ANSITEXTVIEW(": "));
 			}
 
 			OutputStream.Append(Block.GetCode());
@@ -1515,20 +1509,20 @@ static void BuildLineBreakMap(FShaderSource::FViewType Source, TArray<int32, FMe
 	int32 Cursor = 0;
 
 #if UE_SHADER_MINIFIER_SSE
-	const __m128i Needle = FShaderSource::IsWide() ? _mm_set1_epi16(L'\n') : _mm_set1_epi8('\n');
+	const __m128i Needle = _mm_set1_epi8('\n');
 	while (Cursor < SourceLen)
 	{
 		__m128i Chunk = _mm_loadu_si128(reinterpret_cast<const __m128i*>(Chars + Cursor));
-		__m128i MaskVec = FShaderSource::IsWide() ? _mm_cmpeq_epi16(Chunk, Needle) : _mm_cmpeq_epi8(Chunk, Needle);
+		__m128i MaskVec = _mm_cmpeq_epi8(Chunk, Needle);
 		uint32 Mask = _mm_movemask_epi8(MaskVec);
 		while (Mask != 0)
 		{
 			const uint32 BitIndex = FMath::CountTrailingZeros(Mask);
 			const uint32 ChunkCharIndex = BitIndex / sizeof(FShaderSource::CharType);
 			OutLineBreakMap.Add(Cursor + ChunkCharIndex);
-			Mask &= ~(FShaderSource::GetSingleCharMask() << BitIndex);
+			Mask &= ~(1u << BitIndex);
 		}
-		Cursor += FShaderSource::GetSimdCharCount();
+		Cursor += 16;
 	}
 #else
 	while (Cursor < SourceLen)
@@ -1574,7 +1568,7 @@ static int32 FindLineNumber(FShaderSource::FViewType Source, const TArray<int32,
 
 static bool ParseLineDirective(FShaderSource::FViewType Input, int32& OutLineNumber, FShaderSource::FViewType& OutFileName)
 {
-	if (!StartsWith(Input, SHADER_SOURCE_VIEWLITERAL("#line")))
+	if (!StartsWith(Input, ANSITEXTVIEW("#line")))
 	{
 		return false;
 	}
@@ -1611,9 +1605,9 @@ static void OpenNamespace(FShaderSource::FStringType& OutputStream, const FNames
 {
 	for (const FShaderSource::FViewType& Name : Namespace.Stack)
 	{
-		OutputStream.Append(SHADER_SOURCE_VIEWLITERAL("namespace "));
+		OutputStream.Append(ANSITEXTVIEW("namespace "));
 		OutputStream.Append(Name);
-		OutputStream.Append(SHADER_SOURCE_VIEWLITERAL(" { "));
+		OutputStream.Append(ANSITEXTVIEW(" { "));
 	}
 }
 
@@ -1624,7 +1618,7 @@ static void CloseNamespace(FShaderSource::FStringType& OutputStream, const FName
 		OutputStream.AppendChar('}');
 	}
 
-	OutputStream.Append(SHADER_SOURCE_VIEWLITERAL(" // namespace "));
+	OutputStream.Append(ANSITEXTVIEW(" // namespace "));
 	OutputStream.Append(Namespace.FullName);
 }
 
@@ -1657,141 +1651,141 @@ static FShaderSource::FStringType MinifyShader(const FParsedShader& Parsed, TCon
 
 	{
 		// Some known builtin words to ignore
-		ProcessedIdentifiers.Add(SHADER_SOURCE_LITERAL("asfloat"));
-		ProcessedIdentifiers.Add(SHADER_SOURCE_LITERAL("asint"));
-		ProcessedIdentifiers.Add(SHADER_SOURCE_LITERAL("asuint"));
-		ProcessedIdentifiers.Add(SHADER_SOURCE_LITERAL("bool"));
-		ProcessedIdentifiers.Add(SHADER_SOURCE_LITERAL("bool2"));
-		ProcessedIdentifiers.Add(SHADER_SOURCE_LITERAL("bool3"));
-		ProcessedIdentifiers.Add(SHADER_SOURCE_LITERAL("bool4"));
-		ProcessedIdentifiers.Add(SHADER_SOURCE_LITERAL("break"));
-		ProcessedIdentifiers.Add(SHADER_SOURCE_LITERAL("cbuffer"));
-		ProcessedIdentifiers.Add(SHADER_SOURCE_LITERAL("const"));
-		ProcessedIdentifiers.Add(SHADER_SOURCE_LITERAL("else"));
-		ProcessedIdentifiers.Add(SHADER_SOURCE_LITERAL("extern"));
-		ProcessedIdentifiers.Add(SHADER_SOURCE_LITERAL("false"));
-		ProcessedIdentifiers.Add(SHADER_SOURCE_LITERAL("float"));
-		ProcessedIdentifiers.Add(SHADER_SOURCE_LITERAL("float2"));
-		ProcessedIdentifiers.Add(SHADER_SOURCE_LITERAL("float3"));
-		ProcessedIdentifiers.Add(SHADER_SOURCE_LITERAL("float3x3"));
-		ProcessedIdentifiers.Add(SHADER_SOURCE_LITERAL("float3x4"));
-		ProcessedIdentifiers.Add(SHADER_SOURCE_LITERAL("float4"));
-		ProcessedIdentifiers.Add(SHADER_SOURCE_LITERAL("float4x4"));
-		ProcessedIdentifiers.Add(SHADER_SOURCE_LITERAL("for"));
-		ProcessedIdentifiers.Add(SHADER_SOURCE_LITERAL("groupshared"));
-		ProcessedIdentifiers.Add(SHADER_SOURCE_LITERAL("if"));
-		ProcessedIdentifiers.Add(SHADER_SOURCE_LITERAL("in"));
-		ProcessedIdentifiers.Add(SHADER_SOURCE_LITERAL("inout"));
-		ProcessedIdentifiers.Add(SHADER_SOURCE_LITERAL("int"));
-		ProcessedIdentifiers.Add(SHADER_SOURCE_LITERAL("int2"));
-		ProcessedIdentifiers.Add(SHADER_SOURCE_LITERAL("int3"));
-		ProcessedIdentifiers.Add(SHADER_SOURCE_LITERAL("int4"));
-		ProcessedIdentifiers.Add(SHADER_SOURCE_LITERAL("interface"));
-		ProcessedIdentifiers.Add(SHADER_SOURCE_LITERAL("out"));
-		ProcessedIdentifiers.Add(SHADER_SOURCE_LITERAL("packoffset"));
-		ProcessedIdentifiers.Add(SHADER_SOURCE_LITERAL("precise"));
-		ProcessedIdentifiers.Add(SHADER_SOURCE_LITERAL("register"));
-		ProcessedIdentifiers.Add(SHADER_SOURCE_LITERAL("return"));
-		ProcessedIdentifiers.Add(SHADER_SOURCE_LITERAL("static"));
-		ProcessedIdentifiers.Add(SHADER_SOURCE_LITERAL("struct"));
-		ProcessedIdentifiers.Add(SHADER_SOURCE_LITERAL("switch"));
-		ProcessedIdentifiers.Add(SHADER_SOURCE_LITERAL("tbuffer"));
-		ProcessedIdentifiers.Add(SHADER_SOURCE_LITERAL("true"));
-		ProcessedIdentifiers.Add(SHADER_SOURCE_LITERAL("uint"));
-		ProcessedIdentifiers.Add(SHADER_SOURCE_LITERAL("uint2"));
-		ProcessedIdentifiers.Add(SHADER_SOURCE_LITERAL("uint3"));
-		ProcessedIdentifiers.Add(SHADER_SOURCE_LITERAL("uint4"));
-		ProcessedIdentifiers.Add(SHADER_SOURCE_LITERAL("void"));
-		ProcessedIdentifiers.Add(SHADER_SOURCE_LITERAL("while"));
-		ProcessedIdentifiers.Add(SHADER_SOURCE_LITERAL("typedef"));
-		ProcessedIdentifiers.Add(SHADER_SOURCE_LITERAL("template"));
-		ProcessedIdentifiers.Add(SHADER_SOURCE_LITERAL("operator"));
-		ProcessedIdentifiers.Add(SHADER_SOURCE_LITERAL("enum"));
+		ProcessedIdentifiers.Add("asfloat");
+		ProcessedIdentifiers.Add("asint");
+		ProcessedIdentifiers.Add("asuint");
+		ProcessedIdentifiers.Add("bool");
+		ProcessedIdentifiers.Add("bool2");
+		ProcessedIdentifiers.Add("bool3");
+		ProcessedIdentifiers.Add("bool4");
+		ProcessedIdentifiers.Add("break");
+		ProcessedIdentifiers.Add("cbuffer");
+		ProcessedIdentifiers.Add("const");
+		ProcessedIdentifiers.Add("else");
+		ProcessedIdentifiers.Add("extern");
+		ProcessedIdentifiers.Add("false");
+		ProcessedIdentifiers.Add("float");
+		ProcessedIdentifiers.Add("float2");
+		ProcessedIdentifiers.Add("float3");
+		ProcessedIdentifiers.Add("float3x3");
+		ProcessedIdentifiers.Add("float3x4");
+		ProcessedIdentifiers.Add("float4");
+		ProcessedIdentifiers.Add("float4x4");
+		ProcessedIdentifiers.Add("for");
+		ProcessedIdentifiers.Add("groupshared");
+		ProcessedIdentifiers.Add("if");
+		ProcessedIdentifiers.Add("in");
+		ProcessedIdentifiers.Add("inout");
+		ProcessedIdentifiers.Add("int");
+		ProcessedIdentifiers.Add("int2");
+		ProcessedIdentifiers.Add("int3");
+		ProcessedIdentifiers.Add("int4");
+		ProcessedIdentifiers.Add("interface");
+		ProcessedIdentifiers.Add("out");
+		ProcessedIdentifiers.Add("packoffset");
+		ProcessedIdentifiers.Add("precise");
+		ProcessedIdentifiers.Add("register");
+		ProcessedIdentifiers.Add("return");
+		ProcessedIdentifiers.Add("static");
+		ProcessedIdentifiers.Add("struct");
+		ProcessedIdentifiers.Add("switch");
+		ProcessedIdentifiers.Add("tbuffer");
+		ProcessedIdentifiers.Add("true");
+		ProcessedIdentifiers.Add("uint");
+		ProcessedIdentifiers.Add("uint2");
+		ProcessedIdentifiers.Add("uint3");
+		ProcessedIdentifiers.Add("uint4");
+		ProcessedIdentifiers.Add("void");
+		ProcessedIdentifiers.Add("while");
+		ProcessedIdentifiers.Add("typedef");
+		ProcessedIdentifiers.Add("template");
+		ProcessedIdentifiers.Add("operator");
+		ProcessedIdentifiers.Add("enum");
 
 		// HLSL resource types
-		ProcessedIdentifiers.Add(SHADER_SOURCE_LITERAL("TextureCubeArray"));
-		ProcessedIdentifiers.Add(SHADER_SOURCE_LITERAL("TextureCube"));
-		ProcessedIdentifiers.Add(SHADER_SOURCE_LITERAL("TextureBuffer"));
-		ProcessedIdentifiers.Add(SHADER_SOURCE_LITERAL("Texture3D"));
-		ProcessedIdentifiers.Add(SHADER_SOURCE_LITERAL("Texture2DMSArray"));
-		ProcessedIdentifiers.Add(SHADER_SOURCE_LITERAL("Texture2DMS"));
-		ProcessedIdentifiers.Add(SHADER_SOURCE_LITERAL("Texture2DArray"));
-		ProcessedIdentifiers.Add(SHADER_SOURCE_LITERAL("Texture2D"));
-		ProcessedIdentifiers.Add(SHADER_SOURCE_LITERAL("Texture1DArray"));
-		ProcessedIdentifiers.Add(SHADER_SOURCE_LITERAL("Texture1D"));
-		ProcessedIdentifiers.Add(SHADER_SOURCE_LITERAL("StructuredBuffer"));
-		ProcessedIdentifiers.Add(SHADER_SOURCE_LITERAL("SamplerState"));
-		ProcessedIdentifiers.Add(SHADER_SOURCE_LITERAL("SamplerComparisonState"));
-		ProcessedIdentifiers.Add(SHADER_SOURCE_LITERAL("RWTextureCubeArray"));
-		ProcessedIdentifiers.Add(SHADER_SOURCE_LITERAL("RWTextureCube"));
-		ProcessedIdentifiers.Add(SHADER_SOURCE_LITERAL("RWTexture3D"));
-		ProcessedIdentifiers.Add(SHADER_SOURCE_LITERAL("RWTexture2DMSArray"));
-		ProcessedIdentifiers.Add(SHADER_SOURCE_LITERAL("RWTexture2DMS"));
-		ProcessedIdentifiers.Add(SHADER_SOURCE_LITERAL("RWTexture2DArray"));
-		ProcessedIdentifiers.Add(SHADER_SOURCE_LITERAL("RWTexture2D"));
-		ProcessedIdentifiers.Add(SHADER_SOURCE_LITERAL("RWTexture1DArray"));
-		ProcessedIdentifiers.Add(SHADER_SOURCE_LITERAL("RWTexture1D"));
-		ProcessedIdentifiers.Add(SHADER_SOURCE_LITERAL("RWStructuredBuffer"));
-		ProcessedIdentifiers.Add(SHADER_SOURCE_LITERAL("RWByteAddressBuffer"));
-		ProcessedIdentifiers.Add(SHADER_SOURCE_LITERAL("RWBuffer"));
-		ProcessedIdentifiers.Add(SHADER_SOURCE_LITERAL("RaytracingAccelerationStructure"));
-		ProcessedIdentifiers.Add(SHADER_SOURCE_LITERAL("RasterizerOrderedTexture3D"));
-		ProcessedIdentifiers.Add(SHADER_SOURCE_LITERAL("RasterizerOrderedTexture2DArray"));
-		ProcessedIdentifiers.Add(SHADER_SOURCE_LITERAL("RasterizerOrderedTexture2D"));
-		ProcessedIdentifiers.Add(SHADER_SOURCE_LITERAL("RasterizerOrderedTexture1DArray"));
-		ProcessedIdentifiers.Add(SHADER_SOURCE_LITERAL("RasterizerOrderedTexture1D"));
-		ProcessedIdentifiers.Add(SHADER_SOURCE_LITERAL("RasterizerOrderedStructuredBuffer"));
-		ProcessedIdentifiers.Add(SHADER_SOURCE_LITERAL("RasterizerOrderedByteAddressBuffer"));
-		ProcessedIdentifiers.Add(SHADER_SOURCE_LITERAL("RasterizerOrderedBuffer"));
-		ProcessedIdentifiers.Add(SHADER_SOURCE_LITERAL("FeedbackTexture2DArray"));
-		ProcessedIdentifiers.Add(SHADER_SOURCE_LITERAL("FeedbackTexture2D"));
-		ProcessedIdentifiers.Add(SHADER_SOURCE_LITERAL("ConsumeStructuredBuffer"));
-		ProcessedIdentifiers.Add(SHADER_SOURCE_LITERAL("ConstantBuffer"));
-		ProcessedIdentifiers.Add(SHADER_SOURCE_LITERAL("ByteAddressBuffer"));
-		ProcessedIdentifiers.Add(SHADER_SOURCE_LITERAL("Buffer"));
-		ProcessedIdentifiers.Add(SHADER_SOURCE_LITERAL("AppendStructuredBuffer"));
+		ProcessedIdentifiers.Add("TextureCubeArray");
+		ProcessedIdentifiers.Add("TextureCube");
+		ProcessedIdentifiers.Add("TextureBuffer");
+		ProcessedIdentifiers.Add("Texture3D");
+		ProcessedIdentifiers.Add("Texture2DMSArray");
+		ProcessedIdentifiers.Add("Texture2DMS");
+		ProcessedIdentifiers.Add("Texture2DArray");
+		ProcessedIdentifiers.Add("Texture2D");
+		ProcessedIdentifiers.Add("Texture1DArray");
+		ProcessedIdentifiers.Add("Texture1D");
+		ProcessedIdentifiers.Add("StructuredBuffer");
+		ProcessedIdentifiers.Add("SamplerState");
+		ProcessedIdentifiers.Add("SamplerComparisonState");
+		ProcessedIdentifiers.Add("RWTextureCubeArray");
+		ProcessedIdentifiers.Add("RWTextureCube");
+		ProcessedIdentifiers.Add("RWTexture3D");
+		ProcessedIdentifiers.Add("RWTexture2DMSArray");
+		ProcessedIdentifiers.Add("RWTexture2DMS");
+		ProcessedIdentifiers.Add("RWTexture2DArray");
+		ProcessedIdentifiers.Add("RWTexture2D");
+		ProcessedIdentifiers.Add("RWTexture1DArray");
+		ProcessedIdentifiers.Add("RWTexture1D");
+		ProcessedIdentifiers.Add("RWStructuredBuffer");
+		ProcessedIdentifiers.Add("RWByteAddressBuffer");
+		ProcessedIdentifiers.Add("RWBuffer");
+		ProcessedIdentifiers.Add("RaytracingAccelerationStructure");
+		ProcessedIdentifiers.Add("RasterizerOrderedTexture3D");
+		ProcessedIdentifiers.Add("RasterizerOrderedTexture2DArray");
+		ProcessedIdentifiers.Add("RasterizerOrderedTexture2D");
+		ProcessedIdentifiers.Add("RasterizerOrderedTexture1DArray");
+		ProcessedIdentifiers.Add("RasterizerOrderedTexture1D");
+		ProcessedIdentifiers.Add("RasterizerOrderedStructuredBuffer");
+		ProcessedIdentifiers.Add("RasterizerOrderedByteAddressBuffer");
+		ProcessedIdentifiers.Add("RasterizerOrderedBuffer");
+		ProcessedIdentifiers.Add("FeedbackTexture2DArray");
+		ProcessedIdentifiers.Add("FeedbackTexture2D");
+		ProcessedIdentifiers.Add("ConsumeStructuredBuffer");
+		ProcessedIdentifiers.Add("ConstantBuffer");
+		ProcessedIdentifiers.Add("ByteAddressBuffer");
+		ProcessedIdentifiers.Add("Buffer");
+		ProcessedIdentifiers.Add("AppendStructuredBuffer");
 
 		// Alternative spelling of some resource types
-		ProcessedIdentifiers.Add(SHADER_SOURCE_LITERAL("AppendRegularBuffer"));
-		ProcessedIdentifiers.Add(SHADER_SOURCE_LITERAL("ByteBuffer"));
-		ProcessedIdentifiers.Add(SHADER_SOURCE_LITERAL("ConsumeRegularBuffer"));
-		ProcessedIdentifiers.Add(SHADER_SOURCE_LITERAL("DataBuffer"));
-		ProcessedIdentifiers.Add(SHADER_SOURCE_LITERAL("MS_Texture2D"));
-		ProcessedIdentifiers.Add(SHADER_SOURCE_LITERAL("MS_Texture2D_Array"));
-		ProcessedIdentifiers.Add(SHADER_SOURCE_LITERAL("RegularBuffer"));
-		ProcessedIdentifiers.Add(SHADER_SOURCE_LITERAL("RW_ByteBuffer"));
-		ProcessedIdentifiers.Add(SHADER_SOURCE_LITERAL("RW_DataBuffer"));
-		ProcessedIdentifiers.Add(SHADER_SOURCE_LITERAL("RW_RegularBuffer"));
-		ProcessedIdentifiers.Add(SHADER_SOURCE_LITERAL("RW_Texture1D"));
-		ProcessedIdentifiers.Add(SHADER_SOURCE_LITERAL("RW_Texture1D_Array"));
-		ProcessedIdentifiers.Add(SHADER_SOURCE_LITERAL("RW_Texture2D"));
-		ProcessedIdentifiers.Add(SHADER_SOURCE_LITERAL("RW_Texture2D_Array"));
-		ProcessedIdentifiers.Add(SHADER_SOURCE_LITERAL("RW_Texture3D"));
-		ProcessedIdentifiers.Add(SHADER_SOURCE_LITERAL("RW_TextureCube"));
-		ProcessedIdentifiers.Add(SHADER_SOURCE_LITERAL("Texture1D_Array"));
-		ProcessedIdentifiers.Add(SHADER_SOURCE_LITERAL("Texture2D_Array"));
-		ProcessedIdentifiers.Add(SHADER_SOURCE_LITERAL("TextureBuffer"));
-		ProcessedIdentifiers.Add(SHADER_SOURCE_LITERAL("TextureCube_Array"));
+		ProcessedIdentifiers.Add("AppendRegularBuffer");
+		ProcessedIdentifiers.Add("ByteBuffer");
+		ProcessedIdentifiers.Add("ConsumeRegularBuffer");
+		ProcessedIdentifiers.Add("DataBuffer");
+		ProcessedIdentifiers.Add("MS_Texture2D");
+		ProcessedIdentifiers.Add("MS_Texture2D_Array");
+		ProcessedIdentifiers.Add("RegularBuffer");
+		ProcessedIdentifiers.Add("RW_ByteBuffer");
+		ProcessedIdentifiers.Add("RW_DataBuffer");
+		ProcessedIdentifiers.Add("RW_RegularBuffer");
+		ProcessedIdentifiers.Add("RW_Texture1D");
+		ProcessedIdentifiers.Add("RW_Texture1D_Array");
+		ProcessedIdentifiers.Add("RW_Texture2D");
+		ProcessedIdentifiers.Add("RW_Texture2D_Array");
+		ProcessedIdentifiers.Add("RW_Texture3D");
+		ProcessedIdentifiers.Add("RW_TextureCube");
+		ProcessedIdentifiers.Add("Texture1D_Array");
+		ProcessedIdentifiers.Add("Texture2D_Array");
+		ProcessedIdentifiers.Add("TextureBuffer");
+		ProcessedIdentifiers.Add("TextureCube_Array");
 
 		// Some shaders define template versions of some built-in functions, so we can't trivially ignore them
-		//ProcessedIdentifiers.Add(SHADER_SOURCE_LITERAL("abs"));
-		//ProcessedIdentifiers.Add(SHADER_SOURCE_LITERAL("any"));
-		//ProcessedIdentifiers.Add(SHADER_SOURCE_LITERAL("clamp"));
-		//ProcessedIdentifiers.Add(SHADER_SOURCE_LITERAL("clip"));
-		//ProcessedIdentifiers.Add(SHADER_SOURCE_LITERAL("cos"));
-		//ProcessedIdentifiers.Add(SHADER_SOURCE_LITERAL("cross"));
-		//ProcessedIdentifiers.Add(SHADER_SOURCE_LITERAL("dot"));
-		//ProcessedIdentifiers.Add(SHADER_SOURCE_LITERAL("frac"));
-		//ProcessedIdentifiers.Add(SHADER_SOURCE_LITERAL("lerp"));
-		//ProcessedIdentifiers.Add(SHADER_SOURCE_LITERAL("max"));
-		//ProcessedIdentifiers.Add(SHADER_SOURCE_LITERAL("min"));
-		//ProcessedIdentifiers.Add(SHADER_SOURCE_LITERAL("mul"));
-		//ProcessedIdentifiers.Add(SHADER_SOURCE_LITERAL("normalize"));
-		//ProcessedIdentifiers.Add(SHADER_SOURCE_LITERAL("pow"));
-		//ProcessedIdentifiers.Add(SHADER_SOURCE_LITERAL("saturate"));
-		//ProcessedIdentifiers.Add(SHADER_SOURCE_LITERAL("sign"));
-		//ProcessedIdentifiers.Add(SHADER_SOURCE_LITERAL("sin"));
-		//ProcessedIdentifiers.Add(SHADER_SOURCE_LITERAL("sqrt"));
+		//ProcessedIdentifiers.Add("abs");
+		//ProcessedIdentifiers.Add("any");
+		//ProcessedIdentifiers.Add("clamp");
+		//ProcessedIdentifiers.Add("clip");
+		//ProcessedIdentifiers.Add("cos");
+		//ProcessedIdentifiers.Add("cross");
+		//ProcessedIdentifiers.Add("dot");
+		//ProcessedIdentifiers.Add("frac");
+		//ProcessedIdentifiers.Add("lerp");
+		//ProcessedIdentifiers.Add("max");
+		//ProcessedIdentifiers.Add("min");
+		//ProcessedIdentifiers.Add("mul");
+		//ProcessedIdentifiers.Add("normalize");
+		//ProcessedIdentifiers.Add("pow");
+		//ProcessedIdentifiers.Add("saturate");
+		//ProcessedIdentifiers.Add("sign");
+		//ProcessedIdentifiers.Add("sin");
+		//ProcessedIdentifiers.Add("sqrt");
 	}
 
 	if (PendingChunks.IsEmpty())
@@ -1939,27 +1933,27 @@ static FShaderSource::FStringType MinifyShader(const FParsedShader& Parsed, TCon
 
 	if (EnumHasAnyFlags(Flags, EMinifyShaderFlags::OutputStats))
 	{
-		OutputStream.Append(SHADER_SOURCE_VIEWLITERAL("// Total code chunks: "));
+		OutputStream.Append(ANSITEXTVIEW("// Total code chunks: "));
 		OutputStream.AppendInt(RelevantChunks.Num());
 		OutputStream.AppendChar('\n');
 
-		OutputStream.Append(SHADER_SOURCE_VIEWLITERAL("// - Functions: "));
+		OutputStream.Append(ANSITEXTVIEW("// - Functions: "));
 		OutputStream.AppendInt(NumFunctions);
 		OutputStream.AppendChar('\n');
 
-		OutputStream.Append(SHADER_SOURCE_VIEWLITERAL("// - Structs: "));
+		OutputStream.Append(ANSITEXTVIEW("// - Structs: "));
 		OutputStream.AppendInt(NumStructs);
 		OutputStream.AppendChar('\n');
 
-		OutputStream.Append(SHADER_SOURCE_VIEWLITERAL("// - CBuffers: "));
+		OutputStream.Append(ANSITEXTVIEW("// - CBuffers: "));
 		OutputStream.AppendInt(NumCBuffers);
 		OutputStream.AppendChar('\n');
 
-		OutputStream.Append(SHADER_SOURCE_VIEWLITERAL("// - Variables: "));
+		OutputStream.Append(ANSITEXTVIEW("// - Variables: "));
 		OutputStream.AppendInt(NumVariables);
 		OutputStream.AppendChar('\n');
 
-		OutputStream.Append(SHADER_SOURCE_VIEWLITERAL("// - Other: "));
+		OutputStream.Append(ANSITEXTVIEW("// - Other: "));
 		OutputStream.AppendInt(NumOtherChunks);
 		OutputStream.AppendChar('\n');
 
@@ -2020,13 +2014,13 @@ static FShaderSource::FStringType MinifyShader(const FParsedShader& Parsed, TCon
 			if (CurrentNamespace)
 			{
 				CloseNamespace(OutputStream, *CurrentNamespace);
-				OutputStream.Append(SHADER_SOURCE_VIEWLITERAL("\n\n"));
+				OutputStream.Append(ANSITEXTVIEW("\n\n"));
 			}
 
 			if (PendingNamespace)
 			{
 				OpenNamespace(OutputStream, *PendingNamespace);
-				OutputStream.Append(SHADER_SOURCE_VIEWLITERAL("\n\n"));
+				OutputStream.Append(ANSITEXTVIEW("\n\n"));
 			}
 
 			CurrentNamespace = PendingNamespace;
@@ -2041,7 +2035,7 @@ static FShaderSource::FStringType MinifyShader(const FParsedShader& Parsed, TCon
 				FShaderSource::FViewType  RequestedByName = RequestedByChunk->FindFirstBlockByType(EBlockType::Name);
 				if (!RequestedByName.IsEmpty())
 				{
-					OutputStream.Append(SHADER_SOURCE_VIEWLITERAL("// REASON: "));
+					OutputStream.Append(ANSITEXTVIEW("// REASON: "));
 					OutputStream.Append(RequestedByName);
 					OutputStream.AppendChar('\n');
 				}
@@ -2059,7 +2053,7 @@ static FShaderSource::FStringType MinifyShader(const FParsedShader& Parsed, TCon
 				// There was no valid line directive for this chunk, but we do know the line in the input source, so just emit that.
 				if (ChunkLine > LastLineNumber + 1 || !LastLineFileName.IsEmpty())
 				{
-					OutputStream.Append(SHADER_SOURCE_VIEWLITERAL("#line "));
+					OutputStream.Append(ANSITEXTVIEW("#line "));
 					OutputStream.AppendInt(ChunkLine);
 					OutputStream.AppendChar('\n');
 				}
@@ -2090,12 +2084,12 @@ static FShaderSource::FStringType MinifyShader(const FParsedShader& Parsed, TCon
 							OutputStream.AppendChar('\n');
 						}
 
-						OutputStream.Append(SHADER_SOURCE_VIEWLITERAL("#line "));
+						OutputStream.Append(ANSITEXTVIEW("#line "));
 						OutputStream.AppendInt(PatchedLineNumber);
 
 						if (!ParsedFileName.IsEmpty())
 						{
-							OutputStream.Append(SHADER_SOURCE_VIEWLITERAL(" \""));
+							OutputStream.Append(ANSITEXTVIEW(" \""));
 							OutputStream.Append(ParsedFileName);
 							OutputStream.AppendChar('\"');
 						}
@@ -2165,26 +2159,26 @@ TEST_CASE_NAMED(FShaderMinifierParserTest, "Shaders::ShaderMinifier::Parse", "[E
 	using namespace UE::ShaderMinifier;
 
 	{
-		FShaderSource S(SHADER_SOURCE_LITERAL("  \n\r\f \tHello"));
+		FShaderSource S("  \n\r\f \tHello");
 		TEST_EQUAL(TEXT("SkipSpace"),
 			FString(SkipSpace(S.GetView())),
-			FString(SHADER_SOURCE_LITERAL("Hello")));
+			FString("Hello"));
 	}
 	{
-		FShaderSource S(SHADER_SOURCE_LITERAL("Hello World"));
+		FShaderSource S("Hello World");
 		TEST_EQUAL(TEXT("SkipUntilStr (found)"),
-			FString(SkipUntilStr(S.GetView(), SHADER_SOURCE_LITERAL("World"))),
-			FString(SHADER_SOURCE_LITERAL("World")));
+			FString(SkipUntilStr(S.GetView(), "World")),
+			FString("World"));
 	}
 	{
-		FShaderSource S(SHADER_SOURCE_LITERAL("Hello World"));
+		FShaderSource S("Hello World");
 		TEST_EQUAL(TEXT("SkipUntilStr (not found)"),
-			FString(SkipUntilStr(S.GetView(), SHADER_SOURCE_LITERAL("Blah"))),
+			FString(SkipUntilStr(S.GetView(), "Blah")),
 			FString());
 	}
 
 	{
-		FShaderSource S(SHADER_SOURCE_LITERAL("static const struct { int Blah; } Foo = { 123; };"));
+		FShaderSource S("static const struct { int Blah; } Foo = { 123; };");
 		auto P = ParseShader(S);
 		TEST_IF_EQUAL(TEXT("Anonymous struct variable with initializer, total chunks"), P.Chunks.Num(), 1)
 		{
@@ -2193,7 +2187,7 @@ TEST_CASE_NAMED(FShaderMinifierParserTest, "Shaders::ShaderMinifier::Parse", "[E
 	}
 
 	{
-		FShaderSource S(SHADER_SOURCE_LITERAL("float4 PSMain() : SV_Target { return float4(1,0,0,1); };"));
+		FShaderSource S("float4 PSMain() : SV_Target { return float4(1,0,0,1); };");
 		auto P = ParseShader(S);
 		TEST_IF_EQUAL(TEXT("Pixel shader entry point, total chunks"), P.Chunks.Num(), 1)
 		{
@@ -2204,34 +2198,34 @@ TEST_CASE_NAMED(FShaderMinifierParserTest, "Shaders::ShaderMinifier::Parse", "[E
 	{
 		FMemMark Mark(FMemStack::Get());
 		TArray<FShaderSource::FViewType, FMemStackAllocator> R;
-		FShaderSource S(SHADER_SOURCE_LITERAL("Hello[World]; Foo[0];\n"));
+		FShaderSource S("Hello[World]; Foo[0];\n");
 		ExtractIdentifiers(S.GetView(), R);
 		TEST_IF_EQUAL(TEXT("ExtractIdentifiers1: Num"), R.Num(), 3)
 		{
-			TEST_EQUAL(TEXT("ExtractIdentifiers1: R[0]"), FString(R[0]), SHADER_SOURCE_LITERAL("Hello"));
-			TEST_EQUAL(TEXT("ExtractIdentifiers1: R[1]"), FString(R[1]), SHADER_SOURCE_LITERAL("World"));
-			TEST_EQUAL(TEXT("ExtractIdentifiers1: R[2]"), FString(R[2]), SHADER_SOURCE_LITERAL("Foo"));
+			TEST_EQUAL(TEXT("ExtractIdentifiers1: R[0]"), FString(R[0]), "Hello");
+			TEST_EQUAL(TEXT("ExtractIdentifiers1: R[1]"), FString(R[1]), "World");
+			TEST_EQUAL(TEXT("ExtractIdentifiers1: R[2]"), FString(R[2]), "Foo");
 		}
 	}
 
 	{
 		FMemMark Mark(FMemStack::Get());
 		TArray<FShaderSource::FViewType, FMemStackAllocator> R;
-		FShaderSource S(SHADER_SOURCE_LITERAL("#line 0\nStructuredBuffer<uint4> Blah : register(t0, space123);#line 1\n#pragma foo\n"));
+		FShaderSource S("#line 0\nStructuredBuffer<uint4> Blah : register(t0, space123);#line 1\n#pragma foo\n");
 		ExtractIdentifiers(S.GetView(), R);
 		TEST_IF_EQUAL(TEXT("ExtractIdentifiers2: Num"), R.Num(), 6)
 		{
-			TEST_EQUAL(TEXT("ExtractIdentifiers2: R[0]"), FString(R[0]), SHADER_SOURCE_LITERAL("StructuredBuffer"));
-			TEST_EQUAL(TEXT("ExtractIdentifiers2: R[1]"), FString(R[1]), SHADER_SOURCE_LITERAL("uint4"));
-			TEST_EQUAL(TEXT("ExtractIdentifiers2: R[2]"), FString(R[2]), SHADER_SOURCE_LITERAL("Blah"));
-			TEST_EQUAL(TEXT("ExtractIdentifiers2: R[3]"), FString(R[3]), SHADER_SOURCE_LITERAL("register"));
-			TEST_EQUAL(TEXT("ExtractIdentifiers2: R[4]"), FString(R[4]), SHADER_SOURCE_LITERAL("t0"));
-			TEST_EQUAL(TEXT("ExtractIdentifiers2: R[5]"), FString(R[5]), SHADER_SOURCE_LITERAL("space123"));
+			TEST_EQUAL(TEXT("ExtractIdentifiers2: R[0]"), FString(R[0]), "StructuredBuffer");
+			TEST_EQUAL(TEXT("ExtractIdentifiers2: R[1]"), FString(R[1]), "uint4");
+			TEST_EQUAL(TEXT("ExtractIdentifiers2: R[2]"), FString(R[2]), "Blah");
+			TEST_EQUAL(TEXT("ExtractIdentifiers2: R[3]"), FString(R[3]), "register");
+			TEST_EQUAL(TEXT("ExtractIdentifiers2: R[4]"), FString(R[4]), "t0");
+			TEST_EQUAL(TEXT("ExtractIdentifiers2: R[5]"), FString(R[5]), "space123");
 		}
 	}
 
 	{
-		FShaderSource S(SHADER_SOURCE_LITERAL("StructuredBuffer<uint4> Blah : register(t0, space123);"));
+		FShaderSource S("StructuredBuffer<uint4> Blah : register(t0, space123);");
 		auto P = ParseShader(S);
 		TEST_IF_EQUAL(TEXT("ParseShader: structured buffer: num chunks"), P.Chunks.Num(), 1)
 		{
@@ -2240,7 +2234,7 @@ TEST_CASE_NAMED(FShaderMinifierParserTest, "Shaders::ShaderMinifier::Parse", "[E
 	}
 
 	{
-		FShaderSource S(SHADER_SOURCE_LITERAL("const float Foo = 123.45f;"));
+		FShaderSource S("const float Foo = 123.45f;");
 		auto P = ParseShader(S);
 		TEST_IF_EQUAL(TEXT("ParseShader: const float with initializer: num chunks"), P.Chunks.Num(), 1)
 		{
@@ -2249,7 +2243,7 @@ TEST_CASE_NAMED(FShaderMinifierParserTest, "Shaders::ShaderMinifier::Parse", "[E
 	}
 
 	{
-		FShaderSource S(SHADER_SOURCE_LITERAL("struct Blah { int A; };"));
+		FShaderSource S("struct Blah { int A; };");
 		auto P = ParseShader(S);
 		TEST_IF_EQUAL(TEXT("ParseShader: struct: num chunks"), P.Chunks.Num(), 1)
 		{
@@ -2258,7 +2252,7 @@ TEST_CASE_NAMED(FShaderMinifierParserTest, "Shaders::ShaderMinifier::Parse", "[E
 	}
 
 	{
-		FShaderSource S(SHADER_SOURCE_LITERAL("struct Foo { int FooA; }; struct Bar : Foo { int BarA; };"));
+		FShaderSource S("struct Foo { int FooA; }; struct Bar : Foo { int BarA; };");
 		auto P = ParseShader(S);
 		TEST_IF_EQUAL(TEXT("ParseShader: inherited struct: num chunks"), P.Chunks.Num(), 2)
 		{
@@ -2268,7 +2262,7 @@ TEST_CASE_NAMED(FShaderMinifierParserTest, "Shaders::ShaderMinifier::Parse", "[E
 	}
 
 	{
-		FShaderSource S(SHADER_SOURCE_LITERAL("[numthreads(8,8,1)] void Main() {};"));
+		FShaderSource S("[numthreads(8,8,1)] void Main() {};");
 		auto P = ParseShader(S);
 		TEST_IF_EQUAL(TEXT("ParseShader: compute shader entry point: num chunks"), P.Chunks.Num(), 1)
 		{
@@ -2282,7 +2276,7 @@ TEST_CASE_NAMED(FShaderMinifierParserTest, "Shaders::ShaderMinifier::Parse", "[E
 	}
 
 	{
-		FShaderSource S(SHADER_SOURCE_LITERAL("Texture2D Blah : register(t0);"));
+		FShaderSource S("Texture2D Blah : register(t0);");
 		auto P = ParseShader(S);
 		TEST_IF_EQUAL(TEXT("ParseShader: texture with register: num chunks"), P.Chunks.Num(), 1)
 		{
@@ -2290,7 +2284,7 @@ TEST_CASE_NAMED(FShaderMinifierParserTest, "Shaders::ShaderMinifier::Parse", "[E
 		}
 	}
 	{
-		FShaderSource S(SHADER_SOURCE_LITERAL("Texture2D Blah;"));
+		FShaderSource S("Texture2D Blah;");
 		auto P = ParseShader(S);
 		TEST_IF_EQUAL(TEXT("ParseShader: texture: num chunks"), P.Chunks.Num(), 1)
 		{
@@ -2299,7 +2293,7 @@ TEST_CASE_NAMED(FShaderMinifierParserTest, "Shaders::ShaderMinifier::Parse", "[E
 	}
 
 	{
-		FShaderSource S(SHADER_SOURCE_LITERAL("SamplerState Blah : register(s0, space123);"));
+		FShaderSource S("SamplerState Blah : register(s0, space123);");
 		auto P = ParseShader(S);
 		TEST_IF_EQUAL(TEXT("ParseShader: sampler state with register: num chunks"), P.Chunks.Num(), 1)
 		{
@@ -2310,14 +2304,14 @@ TEST_CASE_NAMED(FShaderMinifierParserTest, "Shaders::ShaderMinifier::Parse", "[E
 #if 0
 	{
 		// TODO: handle function forward declarations
-		FShaderSource S(SHADER_SOURCE_LITERAL("Foo Fun(int a);"));
+		FShaderSource S("Foo Fun(int a);");
 		auto P = ParseShader(S);
 		TEST_EQUAL(TEXT("ParseShader: function forward declaration"), P.Chunks[0].Type, ECodeChunkType::FunctionDecl);
 	}
 #endif
 
 	{
-		FShaderSource S(SHADER_SOURCE_LITERAL("void Fun(int a) {};"));
+		FShaderSource S("void Fun(int a) {};");
 		auto P = ParseShader(S);
 		TEST_IF_EQUAL(TEXT("ParseShader: function with trailing semicolon: num chunks"), P.Chunks.Num(), 1)
 		{
@@ -2326,7 +2320,7 @@ TEST_CASE_NAMED(FShaderMinifierParserTest, "Shaders::ShaderMinifier::Parse", "[E
 	}
 
 	{
-		FShaderSource S(SHADER_SOURCE_LITERAL("void Fun(int a) {}"));
+		FShaderSource S("void Fun(int a) {}");
 		auto P = ParseShader(S);
 		TEST_IF_EQUAL(TEXT("ParseShader: function: num chunks"), P.Chunks.Num(), 1)
 		{
@@ -2335,7 +2329,7 @@ TEST_CASE_NAMED(FShaderMinifierParserTest, "Shaders::ShaderMinifier::Parse", "[E
 	}
 
 	{
-		FShaderSource S(SHADER_SOURCE_LITERAL("cbuffer Foo {blah} SamplerState S;"));
+		FShaderSource S("cbuffer Foo {blah} SamplerState S;");
 		auto P = ParseShader(S);
 
 		TEST_IF_EQUAL(TEXT("ParseShader: cbuffer and sampler state: num chunks"), P.Chunks.Num(), 2)
@@ -2346,7 +2340,7 @@ TEST_CASE_NAMED(FShaderMinifierParserTest, "Shaders::ShaderMinifier::Parse", "[E
 	}
 
 	{
-		FShaderSource S(SHADER_SOURCE_LITERAL("struct Foo { int a; };"));
+		FShaderSource S("struct Foo { int a; };");
 		auto P = ParseShader(S);
 		TEST_IF_EQUAL(TEXT("ParseShader: struct: num chunks"), P.Chunks.Num(), 1)
 		{
@@ -2355,7 +2349,7 @@ TEST_CASE_NAMED(FShaderMinifierParserTest, "Shaders::ShaderMinifier::Parse", "[E
 	}
 
 	{
-		FShaderSource S(SHADER_SOURCE_LITERAL("struct { int a; } Foo = { 123; };"));
+		FShaderSource S("struct { int a; } Foo = { 123; };");
 		auto P = ParseShader(S);
 		TEST_IF_EQUAL(TEXT("ParseShader: anonymous struct with variable and initializer: num chunks"), P.Chunks.Num(), 1)
 		{
@@ -2366,16 +2360,15 @@ TEST_CASE_NAMED(FShaderMinifierParserTest, "Shaders::ShaderMinifier::Parse", "[E
 #if 0
 	{
 		// TODO: handle struct forward declarations
-		FShaderSource S(SHADER_SOURCE_LITERAL("struct Foo;"));
+		FShaderSource S("struct Foo;");
 		auto P = ParseShader(S);
 	}
 #endif
 
 	{
 		FShaderSource S(
-			SHADER_SOURCE_LITERAL(
 				"cbuffer MyBuffer : register(b3)"
-				"{ float4 Element1 : packoffset(c0); float1 Element2 : packoffset(c1); float1 Element3 : packoffset(c1.y); }")
+				"{ float4 Element1 : packoffset(c0); float1 Element2 : packoffset(c1); float1 Element3 : packoffset(c1.y); }"
 			);
 		auto P = ParseShader(S);
 		TEST_IF_EQUAL(TEXT("ParseShader: cbuffer with packoffset: num chunks"), P.Chunks.Num(), 1)
@@ -2385,7 +2378,7 @@ TEST_CASE_NAMED(FShaderMinifierParserTest, "Shaders::ShaderMinifier::Parse", "[E
 	}
 
 	{
-		FShaderSource S(SHADER_SOURCE_LITERAL("static const struct { float4 Param; } Foo;"));
+		FShaderSource S("static const struct { float4 Param; } Foo;");
 		auto P = ParseShader(S);
 		TEST_IF_EQUAL(TEXT("ParseShader: static const anonymous struct with variable: num chunks"), P.Chunks.Num(), 1)
 		{
@@ -2394,7 +2387,7 @@ TEST_CASE_NAMED(FShaderMinifierParserTest, "Shaders::ShaderMinifier::Parse", "[E
 	}
 
 	{
-		FShaderSource S(SHADER_SOURCE_LITERAL("static const struct { float4 Param; } Foo = { FooCB_Param; };"));
+		FShaderSource S("static const struct { float4 Param; } Foo = { FooCB_Param; };");
 		auto P = ParseShader(S);
 		TEST_IF_EQUAL(TEXT("ParseShader: static const anonymous struct with variable and initializer: num chunks"), P.Chunks.Num(), 1)
 		{
@@ -2403,7 +2396,7 @@ TEST_CASE_NAMED(FShaderMinifierParserTest, "Shaders::ShaderMinifier::Parse", "[E
 	}
 
 	{
-		FShaderSource S(SHADER_SOURCE_LITERAL("template <typename T> float Fun(T x) { return (float)x; }"));
+		FShaderSource S("template <typename T> float Fun(T x) { return (float)x; }");
 		auto P = ParseShader(S);
 		TEST_IF_EQUAL(TEXT("ParseShader: template function: num chunks"), P.Chunks.Num(), 1)
 		{
@@ -2412,7 +2405,7 @@ TEST_CASE_NAMED(FShaderMinifierParserTest, "Shaders::ShaderMinifier::Parse", "[E
 	}
 
 	{
-		FShaderSource S(SHADER_SOURCE_LITERAL("enum EFoo { A, B = 123 };"));
+		FShaderSource S("enum EFoo { A, B = 123 };");
 		auto P = ParseShader(S);
 		TEST_IF_EQUAL(TEXT("ParseShader: enum: num chunks"), P.Chunks.Num(), 1)
 		{
@@ -2421,7 +2414,7 @@ TEST_CASE_NAMED(FShaderMinifierParserTest, "Shaders::ShaderMinifier::Parse", "[E
 	}
 
 	{
-		FShaderSource S(SHADER_SOURCE_LITERAL("enum class EFoo { A, B };"));
+		FShaderSource S("enum class EFoo { A, B };");
 		auto P = ParseShader(S);
 		TEST_IF_EQUAL(TEXT("ParseShader: enum class: num chunks"), P.Chunks.Num(), 1)
 		{
@@ -2430,7 +2423,7 @@ TEST_CASE_NAMED(FShaderMinifierParserTest, "Shaders::ShaderMinifier::Parse", "[E
 	}
 
 	{
-		FShaderSource S(SHADER_SOURCE_LITERAL("#define Foo 123"));
+		FShaderSource S("#define Foo 123");
 		auto P = ParseShader(S);
 		TEST_IF_EQUAL(TEXT("ParseShader: define: num chunks"), P.Chunks.Num(), 1)
 		{
@@ -2439,7 +2432,7 @@ TEST_CASE_NAMED(FShaderMinifierParserTest, "Shaders::ShaderMinifier::Parse", "[E
 	}
 
 	{
-		FShaderSource S(SHADER_SOURCE_LITERAL("#pragma Foo"));
+		FShaderSource S("#pragma Foo");
 		auto P = ParseShader(S);
 		TEST_IF_EQUAL(TEXT("ParseShader: pragma: num chunks"), P.Chunks.Num(), 1)
 		{
@@ -2448,7 +2441,7 @@ TEST_CASE_NAMED(FShaderMinifierParserTest, "Shaders::ShaderMinifier::Parse", "[E
 	}
 
 	{
-		FShaderSource S(SHADER_SOURCE_LITERAL("ConstantBuffer<Foo> CB : register ( b123, space456);"));
+		FShaderSource S("ConstantBuffer<Foo> CB : register ( b123, space456);");
 		auto P = ParseShader(S);
 		TEST_IF_EQUAL(TEXT("ParseShader: ConstantBuffer<Foo>: num chunks"), P.Chunks.Num(), 1)
 		{
@@ -2457,7 +2450,7 @@ TEST_CASE_NAMED(FShaderMinifierParserTest, "Shaders::ShaderMinifier::Parse", "[E
 	}
 
 	{
-		FShaderSource S(SHADER_SOURCE_LITERAL("namespace NS1 { void Fun() {}; } namespace NS2 { void Fun() {}; }"));
+		FShaderSource S("namespace NS1 { void Fun() {}; } namespace NS2 { void Fun() {}; }");
 		auto P = ParseShader(S);
 		TEST_IF_EQUAL(TEXT("ParseShader: namespaces: num chunks"), P.Chunks.Num(), 2)
 		{
@@ -2470,7 +2463,7 @@ TEST_CASE_NAMED(FShaderMinifierParserTest, "Shaders::ShaderMinifier::Parse", "[E
 	}
 
 	{
-		FShaderSource S(SHADER_SOURCE_LITERAL("template< typename T > TMyStruct<T> operator + ( TMyStruct<T> A, T B ) { /*...*/ }"));
+		FShaderSource S("template< typename T > TMyStruct<T> operator + ( TMyStruct<T> A, T B ) { /*...*/ }");
 		auto P = ParseShader(S);
 		TEST_IF_EQUAL(TEXT("ParseShader: operators: num chunks"), P.Chunks.Num(), 1)
 		{
@@ -2478,19 +2471,19 @@ TEST_CASE_NAMED(FShaderMinifierParserTest, "Shaders::ShaderMinifier::Parse", "[E
 			TEST_IF_EQUAL(TEXT("ParseShader: operators: chunk 0: num blocks"), P.Chunks[0].Blocks.Num(), 8)
 			{
 				FShaderSource::FStringType Args = FShaderSource::FStringType(P.Chunks[0].FindFirstBlockByType(EBlockType::Args));
-				TEST_EQUAL(TEXT("ParseShader: operators: chunk 0: type name"), Args, SHADER_SOURCE_LITERAL("( TMyStruct<T> A, T B )"));
+				TEST_EQUAL(TEXT("ParseShader: operators: chunk 0: type name"), Args, "( TMyStruct<T> A, T B )");
 
 				FShaderSource::FStringType TypeName = FShaderSource::FStringType(P.Chunks[0].FindFirstBlockByType(EBlockType::Type));
-				TEST_EQUAL(TEXT("ParseShader: operators: chunk 0: type name"), TypeName, SHADER_SOURCE_LITERAL("TMyStruct"));
+				TEST_EQUAL(TEXT("ParseShader: operators: chunk 0: type name"), TypeName, "TMyStruct");
 
 				FShaderSource::FStringType OperatorName = FShaderSource::FStringType(P.Chunks[0].FindFirstBlockByType(EBlockType::OperatorName));
-				TEST_EQUAL(TEXT("ParseShader: operators: chunk 0: operator name"), OperatorName, SHADER_SOURCE_LITERAL("+"));
+				TEST_EQUAL(TEXT("ParseShader: operators: chunk 0: operator name"), OperatorName, "+");
 			}
 		}
 	}
 
 	{
-		FShaderSource S(SHADER_SOURCE_LITERAL("typedef Bar Foo;"));
+		FShaderSource S("typedef Bar Foo;");
 		auto P = ParseShader(S);
 		TEST_IF_EQUAL(TEXT("ParseShader: standard typedef: num chunks"), P.Chunks.Num(), 1)
 		{
@@ -2499,7 +2492,7 @@ TEST_CASE_NAMED(FShaderMinifierParserTest, "Shaders::ShaderMinifier::Parse", "[E
 	}
 
 	{
-		FShaderSource S(SHADER_SOURCE_LITERAL("typedef Bar Foo; static const Foo = Bar(0);"));
+		FShaderSource S("typedef Bar Foo; static const Foo = Bar(0);");
 		auto P = ParseShader(S);
 		TEST_IF_EQUAL(TEXT("ParseShader: standard typedef: num chunks"), P.Chunks.Num(), 2)
 		{
@@ -2509,7 +2502,7 @@ TEST_CASE_NAMED(FShaderMinifierParserTest, "Shaders::ShaderMinifier::Parse", "[E
 	}
 
 	{
-		FShaderSource S(SHADER_SOURCE_LITERAL("[Shader(\"node\")] [NodeID(\"WorkgraphNodeArray\", 0)] void WorkgraphNodeArray_0() {}"));
+		FShaderSource S("[Shader(\"node\")] [NodeID(\"WorkgraphNodeArray\", 0)] void WorkgraphNodeArray_0() {}");
 		auto P = ParseShader(S);
 		TEST_IF_EQUAL(TEXT("ParseShader: workgraph node: num chunks"), P.Chunks.Num(), 1)
 		{
@@ -2527,7 +2520,7 @@ TEST_CASE_NAMED(FShaderMinifierMinifyTest, "Shaders::ShaderMinifier::Minify", "[
 	using namespace UE::ShaderMinifier;
 
 	FShaderSource TestShaderCode(
-		SHADER_SOURCE_LITERAL(R"(// dxc /T cs_6_6 /E MainCS MinifierTest.hlsl 
+		R"(// dxc /T cs_6_6 /E MainCS MinifierTest.hlsl 
 struct FFoo
 {
 	float X;
@@ -2710,7 +2703,7 @@ void WorkgraphNodeArray_0()
 void WorkgraphNodeArray_1()
 {
 }
-)"));
+)");
 
 	auto ChunkPresent = [](const FParsedShader& Parsed, FShaderSource::FViewType Name)
 	{
@@ -2731,69 +2724,69 @@ void WorkgraphNodeArray_1()
 
 	{
 		FDiagnostics Diagnostics;
-		FShaderSource Minified = FShaderSource(MinifyShader(Parsed, SHADER_SOURCE_LITERAL("EmptyFunction"), EMinifyShaderFlags::None, Diagnostics));
+		FShaderSource Minified = FShaderSource(MinifyShader(Parsed, "EmptyFunction", EMinifyShaderFlags::None, Diagnostics));
 		FParsedShader MinifiedParsed = ParseShader(Minified);
 		TEST_IF_EQUAL(TEXT("MinifyShader: EmptyFunction: num chunks"), MinifiedParsed.Chunks.Num(), 7)
 		{
-			TEST_EQUAL(TEXT("MinifyShader: EmptyFunction: pragma"), MinifiedParsed.Chunks[0].GetCode(), SHADER_SOURCE_LITERAL("#pragma test_pragma"));
-			TEST_EQUAL(TEXT("MinifyShader: EmptyFunction: C99 pragma 1"), MinifiedParsed.Chunks[1].GetCode(), SHADER_SOURCE_LITERAL("_Pragma(\"\")"));
-			TEST_EQUAL(TEXT("MinifyShader: EmptyFunction: C99 pragma 2"), MinifiedParsed.Chunks[2].GetCode(), SHADER_SOURCE_LITERAL("_Pragma(\"dxc diagnostic push\")"));
-			TEST_EQUAL(TEXT("MinifyShader: EmptyFunction: C99 pragma 3"), MinifiedParsed.Chunks[3].GetCode(), SHADER_SOURCE_LITERAL("_Pragma(\"dxc diagnostic ignored \\\"-Wall\\\"\") _Pragma(\"dxc diagnostic ignored \\\"-Wconversion\\\"\")"));
-			TEST_EQUAL(TEXT("MinifyShader: EmptyFunction: C99 pragma 4"), MinifiedParsed.Chunks[4].GetCode(), SHADER_SOURCE_LITERAL("_Pragma(\"dxc diagnostic pop\")"));
-			TEST_EQUAL(TEXT("MinifyShader: EmptyFunction: define"), MinifiedParsed.Chunks[5].GetCode(), SHADER_SOURCE_LITERAL("#define COMPILER_DEFINITION_TEST 123"));
-			TEST_EQUAL(TEXT("MinifyShader: EmptyFunction: function"), MinifiedParsed.Chunks[6].GetCode(), SHADER_SOURCE_LITERAL("void EmptyFunction(){}"));
+			TEST_EQUAL(TEXT("MinifyShader: EmptyFunction: pragma"), MinifiedParsed.Chunks[0].GetCode(), "#pragma test_pragma");
+			TEST_EQUAL(TEXT("MinifyShader: EmptyFunction: C99 pragma 1"), MinifiedParsed.Chunks[1].GetCode(), "_Pragma(\"\")");
+			TEST_EQUAL(TEXT("MinifyShader: EmptyFunction: C99 pragma 2"), MinifiedParsed.Chunks[2].GetCode(), "_Pragma(\"dxc diagnostic push\")");
+			TEST_EQUAL(TEXT("MinifyShader: EmptyFunction: C99 pragma 3"), MinifiedParsed.Chunks[3].GetCode(), "_Pragma(\"dxc diagnostic ignored \\\"-Wall\\\"\") _Pragma(\"dxc diagnostic ignored \\\"-Wconversion\\\"\")");
+			TEST_EQUAL(TEXT("MinifyShader: EmptyFunction: C99 pragma 4"), MinifiedParsed.Chunks[4].GetCode(), "_Pragma(\"dxc diagnostic pop\")");
+			TEST_EQUAL(TEXT("MinifyShader: EmptyFunction: define"), MinifiedParsed.Chunks[5].GetCode(), "#define COMPILER_DEFINITION_TEST 123");
+			TEST_EQUAL(TEXT("MinifyShader: EmptyFunction: function"), MinifiedParsed.Chunks[6].GetCode(), "void EmptyFunction(){}");
 		}
 	}
 
 	{
 		FDiagnostics Diagnostics;
-		FShaderSource Minified = FShaderSource(MinifyShader(Parsed, SHADER_SOURCE_LITERAL("MainCS"), EMinifyShaderFlags::OutputReasons, Diagnostics));
+		FShaderSource Minified = FShaderSource(MinifyShader(Parsed, "MainCS", EMinifyShaderFlags::OutputReasons, Diagnostics));
 		FParsedShader MinifiedParsed = ParseShader(Minified);
 
 		// Expect true:
-		TEST_TRUE(TEXT("MinifyShader: MainCS: contains MainCS"), ChunkPresent(MinifiedParsed, SHADER_SOURCE_LITERAL("MainCS")));
-		TEST_TRUE(TEXT("MinifyShader: MainCS: contains FFoo"), ChunkPresent(MinifiedParsed, SHADER_SOURCE_LITERAL("FFoo")));
-		TEST_TRUE(TEXT("MinifyShader: MainCS: contains FBar"), ChunkPresent(MinifiedParsed, SHADER_SOURCE_LITERAL("FBar")));
-		TEST_TRUE(TEXT("MinifyShader: MainCS: contains Sum"), ChunkPresent(MinifiedParsed, SHADER_SOURCE_LITERAL("Sum")));
-		TEST_TRUE(TEXT("MinifyShader: MainCS: contains FunA"), ChunkPresent(MinifiedParsed, SHADER_SOURCE_LITERAL("FunA")));
-		TEST_TRUE(TEXT("MinifyShader: MainCS: contains FunB"), ChunkPresent(MinifiedParsed, SHADER_SOURCE_LITERAL("FunB")));
-		TEST_TRUE(TEXT("MinifyShader: MainCS: contains GAnonymousStruct"), ChunkPresent(MinifiedParsed, SHADER_SOURCE_LITERAL("GAnonymousStruct")));
-		TEST_TRUE(TEXT("MinifyShader: MainCS: contains GStructA"), ChunkPresent(MinifiedParsed, SHADER_SOURCE_LITERAL("GStructA")));
-		TEST_TRUE(TEXT("MinifyShader: MainCS: contains GStructB"), ChunkPresent(MinifiedParsed, SHADER_SOURCE_LITERAL("GStructB")));
-		TEST_TRUE(TEXT("MinifyShader: MainCS: contains GStructC"), ChunkPresent(MinifiedParsed, SHADER_SOURCE_LITERAL("GStructC")));
-		TEST_TRUE(TEXT("MinifyShader: MainCS: contains GInitializedAnonymousStructA"), ChunkPresent(MinifiedParsed, SHADER_SOURCE_LITERAL("GInitializedAnonymousStructA")));
-		TEST_TRUE(TEXT("MinifyShader: MainCS: contains GInitializedAnonymousStructB"), ChunkPresent(MinifiedParsed, SHADER_SOURCE_LITERAL("GInitializedAnonymousStructB")));
-		TEST_TRUE(TEXT("MinifyShader: MainCS: contains OutputBufferType"), ChunkPresent(MinifiedParsed, SHADER_SOURCE_LITERAL("OutputBufferType")));
-		TEST_TRUE(TEXT("MinifyShader: MainCS: contains OutputBuffer"), ChunkPresent(MinifiedParsed, SHADER_SOURCE_LITERAL("OutputBuffer")));
-		TEST_TRUE(TEXT("MinifyShader: MainCS: contains FTypedefUsedStruct"), ChunkPresent(MinifiedParsed, SHADER_SOURCE_LITERAL("FTypedefUsedStruct")));
-		TEST_TRUE(TEXT("MinifyShader: MainCS: contains FTypedefUsed"), ChunkPresent(MinifiedParsed, SHADER_SOURCE_LITERAL("FTypedefUsed")));
-		TEST_TRUE(TEXT("MinifyShader: MainCS: contains FTypedefUsedChained"), ChunkPresent(MinifiedParsed, SHADER_SOURCE_LITERAL("FTypedefUsedChained")));
-		TEST_TRUE(TEXT("MinifyShader: MainCS: contains TypedefUsedBuffer"), ChunkPresent(MinifiedParsed, SHADER_SOURCE_LITERAL("TypedefUsedBuffer")));
-		TEST_TRUE(TEXT("MinifyShader: MainCS: contains struct TUnusedTemplate"), MinifiedParsed.Source.Contains(SHADER_SOURCE_LITERAL("struct TUsedTemplate")));
-		TEST_TRUE(TEXT("MinifyShader: MainCS: contains TUsedTemplate<T> operator*"), MinifiedParsed.Source.Contains(SHADER_SOURCE_LITERAL("TUsedTemplate<T> operator*")));
-		TEST_TRUE(TEXT("MinifyShader: MainCS: contains EEnumUsed"), MinifiedParsed.Source.Contains(SHADER_SOURCE_LITERAL("EEnumUsed")));
-		TEST_TRUE(TEXT("MinifyShader: MainCS: contains ENUM_USED_PART_1"), MinifiedParsed.Source.Contains(SHADER_SOURCE_LITERAL("ENUM_USED_PART_1")));
-		TEST_TRUE(TEXT("MinifyShader: MainCS: contains ENUM_USED_PART_2"), MinifiedParsed.Source.Contains(SHADER_SOURCE_LITERAL("ENUM_USED_PART_2")));
-		TEST_TRUE(TEXT("MinifyShader: MainCS: contains struct FWorkgraphRecord"), MinifiedParsed.Source.Contains(SHADER_SOURCE_LITERAL("struct FWorkgraphRecord")));
-		TEST_TRUE(TEXT("MinifyShader: MainCS: contains void WorkgraphNode()"), MinifiedParsed.Source.Contains(SHADER_SOURCE_LITERAL("void WorkgraphNode()")));
-		TEST_TRUE(TEXT("MinifyShader: MainCS: contains void WorkgraphNodeFunction()"), MinifiedParsed.Source.Contains(SHADER_SOURCE_LITERAL("void WorkgraphNodeFunction()")));
-		TEST_TRUE(TEXT("MinifyShader: MainCS: contains void WorkgraphNodeArray_0()"), MinifiedParsed.Source.Contains(SHADER_SOURCE_LITERAL("void WorkgraphNodeArray_0()")));
-		TEST_TRUE(TEXT("MinifyShader: MainCS: contains void WorkgraphNodeArray_1()"), MinifiedParsed.Source.Contains(SHADER_SOURCE_LITERAL("void WorkgraphNodeArray_1()")));
+		TEST_TRUE(TEXT("MinifyShader: MainCS: contains MainCS"), ChunkPresent(MinifiedParsed, "MainCS"));
+		TEST_TRUE(TEXT("MinifyShader: MainCS: contains FFoo"), ChunkPresent(MinifiedParsed, "FFoo"));
+		TEST_TRUE(TEXT("MinifyShader: MainCS: contains FBar"), ChunkPresent(MinifiedParsed, "FBar"));
+		TEST_TRUE(TEXT("MinifyShader: MainCS: contains Sum"), ChunkPresent(MinifiedParsed, "Sum"));
+		TEST_TRUE(TEXT("MinifyShader: MainCS: contains FunA"), ChunkPresent(MinifiedParsed, "FunA"));
+		TEST_TRUE(TEXT("MinifyShader: MainCS: contains FunB"), ChunkPresent(MinifiedParsed, "FunB"));
+		TEST_TRUE(TEXT("MinifyShader: MainCS: contains GAnonymousStruct"), ChunkPresent(MinifiedParsed, "GAnonymousStruct"));
+		TEST_TRUE(TEXT("MinifyShader: MainCS: contains GStructA"), ChunkPresent(MinifiedParsed, "GStructA"));
+		TEST_TRUE(TEXT("MinifyShader: MainCS: contains GStructB"), ChunkPresent(MinifiedParsed, "GStructB"));
+		TEST_TRUE(TEXT("MinifyShader: MainCS: contains GStructC"), ChunkPresent(MinifiedParsed, "GStructC"));
+		TEST_TRUE(TEXT("MinifyShader: MainCS: contains GInitializedAnonymousStructA"), ChunkPresent(MinifiedParsed, "GInitializedAnonymousStructA"));
+		TEST_TRUE(TEXT("MinifyShader: MainCS: contains GInitializedAnonymousStructB"), ChunkPresent(MinifiedParsed, "GInitializedAnonymousStructB"));
+		TEST_TRUE(TEXT("MinifyShader: MainCS: contains OutputBufferType"), ChunkPresent(MinifiedParsed, "OutputBufferType"));
+		TEST_TRUE(TEXT("MinifyShader: MainCS: contains OutputBuffer"), ChunkPresent(MinifiedParsed, "OutputBuffer"));
+		TEST_TRUE(TEXT("MinifyShader: MainCS: contains FTypedefUsedStruct"), ChunkPresent(MinifiedParsed, "FTypedefUsedStruct"));
+		TEST_TRUE(TEXT("MinifyShader: MainCS: contains FTypedefUsed"), ChunkPresent(MinifiedParsed, "FTypedefUsed"));
+		TEST_TRUE(TEXT("MinifyShader: MainCS: contains FTypedefUsedChained"), ChunkPresent(MinifiedParsed, "FTypedefUsedChained"));
+		TEST_TRUE(TEXT("MinifyShader: MainCS: contains TypedefUsedBuffer"), ChunkPresent(MinifiedParsed, "TypedefUsedBuffer"));
+		TEST_TRUE(TEXT("MinifyShader: MainCS: contains struct TUnusedTemplate"), MinifiedParsed.Source.Contains("struct TUsedTemplate"));
+		TEST_TRUE(TEXT("MinifyShader: MainCS: contains TUsedTemplate<T> operator*"), MinifiedParsed.Source.Contains("TUsedTemplate<T> operator*"));
+		TEST_TRUE(TEXT("MinifyShader: MainCS: contains EEnumUsed"), MinifiedParsed.Source.Contains("EEnumUsed"));
+		TEST_TRUE(TEXT("MinifyShader: MainCS: contains ENUM_USED_PART_1"), MinifiedParsed.Source.Contains("ENUM_USED_PART_1"));
+		TEST_TRUE(TEXT("MinifyShader: MainCS: contains ENUM_USED_PART_2"), MinifiedParsed.Source.Contains("ENUM_USED_PART_2"));
+		TEST_TRUE(TEXT("MinifyShader: MainCS: contains struct FWorkgraphRecord"), MinifiedParsed.Source.Contains("struct FWorkgraphRecord"));
+		TEST_TRUE(TEXT("MinifyShader: MainCS: contains void WorkgraphNode()"), MinifiedParsed.Source.Contains("void WorkgraphNode()"));
+		TEST_TRUE(TEXT("MinifyShader: MainCS: contains void WorkgraphNodeFunction()"), MinifiedParsed.Source.Contains("void WorkgraphNodeFunction()"));
+		TEST_TRUE(TEXT("MinifyShader: MainCS: contains void WorkgraphNodeArray_0()"), MinifiedParsed.Source.Contains("void WorkgraphNodeArray_0()"));
+		TEST_TRUE(TEXT("MinifyShader: MainCS: contains void WorkgraphNodeArray_1()"), MinifiedParsed.Source.Contains("void WorkgraphNodeArray_1()"));
 
 		// Expect false:
-		TEST_FALSE(TEXT("MinifyShader: MainCS: contains UnreferencedFunction"), ChunkPresent(MinifiedParsed, SHADER_SOURCE_LITERAL("UnreferencedFunction")));
-		TEST_FALSE(TEXT("MinifyShader: MainCS: contains FUnreferencedStruct"), ChunkPresent(MinifiedParsed, SHADER_SOURCE_LITERAL("FUnreferencedStruct")));
-		TEST_FALSE(TEXT("MinifyShader: MainCS: contains GUnreferencedParameter"), ChunkPresent(MinifiedParsed, SHADER_SOURCE_LITERAL("GUnreferencedParameter")));
-		TEST_FALSE(TEXT("MinifyShader: MainCS: contains FTypedefUsedChainedUnused"), ChunkPresent(MinifiedParsed, SHADER_SOURCE_LITERAL("FTypedefUsedChainedUnused")));
-		TEST_FALSE(TEXT("MinifyShader: MainCS: contains FTypedefUnusedStruct"), ChunkPresent(MinifiedParsed, SHADER_SOURCE_LITERAL("FTypedefUnusedStruct")));
-		TEST_FALSE(TEXT("MinifyShader: MainCS: contains FTypedefUnused"), ChunkPresent(MinifiedParsed, SHADER_SOURCE_LITERAL("FTypedefUnused")));
-		TEST_FALSE(TEXT("MinifyShader: MainCS: contains TypedefUnusedBuffer"), ChunkPresent(MinifiedParsed, SHADER_SOURCE_LITERAL("TypedefUnusedBuffer")));
-		TEST_FALSE(TEXT("MinifyShader: MainCS: contains struct TUnusedTemplate"), MinifiedParsed.Source.Contains(SHADER_SOURCE_LITERAL("struct TUnusedTemplate")));
-		TEST_FALSE(TEXT("MinifyShader: MainCS: contains TUnusedTemplate<T> operator%"), MinifiedParsed.Source.Contains(SHADER_SOURCE_LITERAL("TUnusedTemplate<T> operator%")));
-		TEST_FALSE(TEXT("MinifyShader: MainCS: contains EEnumUnused"), MinifiedParsed.Source.Contains(SHADER_SOURCE_LITERAL("EEnumUnused")));
-		TEST_FALSE(TEXT("MinifyShader: MainCS: contains ENUM_UNUSED_PART_1"), MinifiedParsed.Source.Contains(SHADER_SOURCE_LITERAL("ENUM_UNUSED_PART_1")));
-		TEST_FALSE(TEXT("MinifyShader: MainCS: contains ENUM_UNUSED_PART_2"), MinifiedParsed.Source.Contains(SHADER_SOURCE_LITERAL("ENUM_UNUSED_PART_2")));
-		TEST_FALSE(TEXT("MinifyShader: MainCS: contains UnreferencedEntryPoint"), MinifiedParsed.Source.Contains(SHADER_SOURCE_LITERAL("UnreferencedEntryPoint")));
+		TEST_FALSE(TEXT("MinifyShader: MainCS: contains UnreferencedFunction"), ChunkPresent(MinifiedParsed, "UnreferencedFunction"));
+		TEST_FALSE(TEXT("MinifyShader: MainCS: contains FUnreferencedStruct"), ChunkPresent(MinifiedParsed, "FUnreferencedStruct"));
+		TEST_FALSE(TEXT("MinifyShader: MainCS: contains GUnreferencedParameter"), ChunkPresent(MinifiedParsed, "GUnreferencedParameter"));
+		TEST_FALSE(TEXT("MinifyShader: MainCS: contains FTypedefUsedChainedUnused"), ChunkPresent(MinifiedParsed, "FTypedefUsedChainedUnused"));
+		TEST_FALSE(TEXT("MinifyShader: MainCS: contains FTypedefUnusedStruct"), ChunkPresent(MinifiedParsed, "FTypedefUnusedStruct"));
+		TEST_FALSE(TEXT("MinifyShader: MainCS: contains FTypedefUnused"), ChunkPresent(MinifiedParsed, "FTypedefUnused"));
+		TEST_FALSE(TEXT("MinifyShader: MainCS: contains TypedefUnusedBuffer"), ChunkPresent(MinifiedParsed, "TypedefUnusedBuffer"));
+		TEST_FALSE(TEXT("MinifyShader: MainCS: contains struct TUnusedTemplate"), MinifiedParsed.Source.Contains("struct TUnusedTemplate"));
+		TEST_FALSE(TEXT("MinifyShader: MainCS: contains TUnusedTemplate<T> operator%"), MinifiedParsed.Source.Contains("TUnusedTemplate<T> operator%"));
+		TEST_FALSE(TEXT("MinifyShader: MainCS: contains EEnumUnused"), MinifiedParsed.Source.Contains("EEnumUnused"));
+		TEST_FALSE(TEXT("MinifyShader: MainCS: contains ENUM_UNUSED_PART_1"), MinifiedParsed.Source.Contains("ENUM_UNUSED_PART_1"));
+		TEST_FALSE(TEXT("MinifyShader: MainCS: contains ENUM_UNUSED_PART_2"), MinifiedParsed.Source.Contains("ENUM_UNUSED_PART_2"));
+		TEST_FALSE(TEXT("MinifyShader: MainCS: contains UnreferencedEntryPoint"), MinifiedParsed.Source.Contains("UnreferencedEntryPoint"));
 	}
 }
 
diff --git a/Engine/Source/Developer/ShaderPreprocessor/Tests/ShaderPreprocessorTests.cpp b/Engine/Source/Developer/ShaderPreprocessor/Tests/ShaderPreprocessorTests.cpp
--- a/Engine/Source/Developer/ShaderPreprocessor/Tests/ShaderPreprocessorTests.cpp
+++ b/Engine/Source/Developer/ShaderPreprocessor/Tests/ShaderPreprocessorTests.cpp
@@ -103,7 +103,7 @@ FPreprocessTestResult ExecutePreprocessTest(const FShaderSource::FViewType& Sour
 TEST_CASE_NAMED(FConditionalExpressionErrorReporting, "Shaders::ShaderPreprocessor::ConditionalExpressionErrorReporting", "[EditorContext][EngineFilter]")
 {
 	// need this dance of multiple assembled string literal fragments to avoid the C++ preprocessor parsing directives inside the literal
-	FShaderSource::FViewType Source = SHADER_SOURCE_VIEWLITERAL(
+	FShaderSource::FViewType Source = ANSITEXTVIEW(
 		"#" "if 0xFFFFFFFF+1 == 2\n"
 		"#" "error \"Unreachable\n\""
 		"#" "endif\n");
@@ -117,7 +117,7 @@ TEST_CASE_NAMED(FInvalidCharError, "Shaders::ShaderPreprocessor::InvalidCharErro
 {
 	{
 		// EOF after #define
-		FShaderSource::FViewType Source = SHADER_SOURCE_VIEWLITERAL(
+		FShaderSource::FViewType Source = ANSITEXTVIEW(
 			"#" "define");
 		FPreprocessTestResult Result = ExecutePreprocessTest(Source);
 		CHECK(Result.NumDiagnostics == 1);
@@ -127,7 +127,7 @@ TEST_CASE_NAMED(FInvalidCharError, "Shaders::ShaderPreprocessor::InvalidCharErro
 
 	{
 		// Invalid printable char in define param list
-		FShaderSource::FViewType Source = SHADER_SOURCE_VIEWLITERAL(
+		FShaderSource::FViewType Source = ANSITEXTVIEW(
 			"#" "define BLAH(abc,#) abc\n");
 		FPreprocessTestResult Result = ExecutePreprocessTest(Source);
 		CHECK(Result.NumDiagnostics == 1);
@@ -137,7 +137,7 @@ TEST_CASE_NAMED(FInvalidCharError, "Shaders::ShaderPreprocessor::InvalidCharErro
 
 	{
 		// Invalid non-printable char in define param list
-		FShaderSource::FViewType Source = SHADER_SOURCE_VIEWLITERAL(
+		FShaderSource::FViewType Source = ANSITEXTVIEW(
 			"#" "define BLAH(abc,\3) abc\n");
 		FPreprocessTestResult Result = ExecutePreprocessTest(Source);
 		CHECK(Result.NumDiagnostics == 1);
@@ -147,7 +147,7 @@ TEST_CASE_NAMED(FInvalidCharError, "Shaders::ShaderPreprocessor::InvalidCharErro
 
 	{
 		// Unexpected char after #if defined(
-		FShaderSource::FViewType Source = SHADER_SOURCE_VIEWLITERAL(
+		FShaderSource::FViewType Source = ANSITEXTVIEW(
 			"#" "define BLAH 1\n"
 			"#" "if defined(BLAH*)\n"
 			"#" "endif\n"
@@ -161,7 +161,7 @@ TEST_CASE_NAMED(FInvalidCharError, "Shaders::ShaderPreprocessor::InvalidCharErro
 
 TEST_CASE_NAMED(FSpecialMacros, "Shaders::ShaderPreprocessor::BuiltinMacros", "[EditorContext][EngineFilter]")
 {
-	FShaderSource::FViewType Source = SHADER_SOURCE_VIEWLITERAL(
+	FShaderSource::FViewType Source = ANSITEXTVIEW(
 		"int A = __" "LINE__;\n"
 		"int B = __" "COUNTER__;\n"
 		"int C = __" "COUNTER__;\n"
@@ -177,7 +177,7 @@ TEST_CASE_NAMED(FEmitLineDirective, "Shaders::ShaderPreprocessor::EmitLineDirect
 {
 	// line directive will be emitted if 12 or more lines are removed due to an inactive preprocessor block
 	// (otherwise the inactive code will be replaced with empty lines)
-	FShaderSource::FViewType Source = SHADER_SOURCE_VIEWLITERAL(
+	FShaderSource::FViewType Source = ANSITEXTVIEW(
 		"#" "define MYDIR 0\n"
 		"#" "if MYDIR\n"
 		"int A = 0;\n"
diff --git a/Engine/Source/Developer/VulkanShaderFormat/Private/SpirVShaderCompiler.inl b/Engine/Source/Developer/VulkanShaderFormat/Private/SpirVShaderCompiler.inl
--- a/Engine/Source/Developer/VulkanShaderFormat/Private/SpirVShaderCompiler.inl
+++ b/Engine/Source/Developer/VulkanShaderFormat/Private/SpirVShaderCompiler.inl
@@ -1174,7 +1174,7 @@ static FString ParseEntrypointDecl(FShaderSource::FViewType PreprocessedShader,
 			continue;
 		}
 
-		int32 ParamsEnd = PreprocessedShader.Find(SHADER_SOURCE_LITERAL(")"), ParamsStart + 1);
+		int32 ParamsEnd = PreprocessedShader.Find(ANSITEXTVIEW(")"), ParamsStart + 1);
 		check(ParamsEnd != INDEX_NONE);
 		if (ParamsEnd == INDEX_NONE)
 		{
diff --git a/Engine/Source/Developer/VulkanShaderFormat/Private/VulkanShaderCompiler.cpp b/Engine/Source/Developer/VulkanShaderFormat/Private/VulkanShaderCompiler.cpp
--- a/Engine/Source/Developer/VulkanShaderFormat/Private/VulkanShaderCompiler.cpp
+++ b/Engine/Source/Developer/VulkanShaderFormat/Private/VulkanShaderCompiler.cpp
@@ -309,14 +309,10 @@ void CompileVulkanShader(const FShaderCompilerInput& Input, const FShaderPreproc
 
 	bool bSuccess = false;
 
-#if SHADER_SOURCE_ANSI
 	// Convert to ANSI prior to calling into ShaderConductor. This copy would have been incurred
 	// by SC itself anyways, but would (will?) also be unnecessary if (when) shader parameter parser
 	// is modified to operate on ANSI strings.
 	const FShaderSource::FStringType PreprocessedSourceToCompile(PreprocessedSource);
-#else
-	const FShaderSource::FStringType& PreprocessedSourceToCompile = PreprocessedSource;
-#endif
 
 #if PLATFORM_MAC || PLATFORM_WINDOWS || PLATFORM_LINUX
 	// HitGroup shaders might have multiple entrypoints that we combine into a single blob
diff --git a/Engine/Source/Runtime/RenderCore/Private/ShaderCore.cpp b/Engine/Source/Runtime/RenderCore/Private/ShaderCore.cpp
--- a/Engine/Source/Runtime/RenderCore/Private/ShaderCore.cpp
+++ b/Engine/Source/Runtime/RenderCore/Private/ShaderCore.cpp
@@ -1650,7 +1650,7 @@ public:
 				FShaderSource::FViewType DebugHashStr = GetShaderSourceDebugHashPrefix();
 				FShaderSource::FViewType SourceView = Job.PreprocessOutput.GetSourceView();
 				int32 DebugHashLoc = SourceView.Find(DebugHashStr) + DebugHashStr.Len();
-				int32 NewlineLoc = SourceView.Find(SHADER_SOURCE_VIEWLITERAL("\n"), DebugHashLoc);
+				int32 NewlineLoc = SourceView.Find(ANSITEXTVIEW("\n"), DebugHashLoc);
 				FShaderSource::TStringBuilder<2 * sizeof(FShaderCompilerInputHash::ByteArray) + 1> HashStr;
 				HashStr << Hash;
 				// handle case where we didn't add the debug hash association (this can be disabled by a cvar for debugging purposes)
diff --git a/Engine/Source/Runtime/RenderCore/Private/ShaderPreprocessTypes.cpp b/Engine/Source/Runtime/RenderCore/Private/ShaderPreprocessTypes.cpp
--- a/Engine/Source/Runtime/RenderCore/Private/ShaderPreprocessTypes.cpp
+++ b/Engine/Source/Runtime/RenderCore/Private/ShaderPreprocessTypes.cpp
@@ -35,18 +35,18 @@ struct FPreprocessConstants
 		StrippedCodeHeader = StrippedCodeHeaderBuilder.ToView();
 	}
 
-	static constexpr FShaderSource::FViewType DebugHashPrefix = SHADER_SOURCE_VIEWLITERAL("DebugHash_");
+	static constexpr FShaderSource::FViewType DebugHashPrefix = ANSITEXTVIEW("DebugHash_");
 	static constexpr FStringView DebugHashPrefixWide = TEXTVIEW("DebugHash_");
-	static constexpr FShaderSource::FViewType FilenameSentinel = SHADER_SOURCE_VIEWLITERAL("__UE_FILENAME_SENTINEL");
+	static constexpr FShaderSource::FViewType FilenameSentinel = ANSITEXTVIEW("__UE_FILENAME_SENTINEL");
 
 	FShaderSource::FViewType StrippedCodeHeader;
 
 private:
 
-	static constexpr FShaderSource::FViewType DebugHashCommentStart = SHADER_SOURCE_VIEWLITERAL("// ");
-	static constexpr FShaderSource::FViewType DebugHashCommentEnd = SHADER_SOURCE_VIEWLITERAL("\n");
-	static constexpr FShaderSource::FViewType LineDirectiveStart = SHADER_SOURCE_VIEWLITERAL("#line 1 \"");
-	static constexpr FShaderSource::FViewType LineDirectiveEnd = SHADER_SOURCE_VIEWLITERAL("\"\n");
+	static constexpr FShaderSource::FViewType DebugHashCommentStart = ANSITEXTVIEW("// ");
+	static constexpr FShaderSource::FViewType DebugHashCommentEnd = ANSITEXTVIEW("\n");
+	static constexpr FShaderSource::FViewType LineDirectiveStart = ANSITEXTVIEW("#line 1 \"");
+	static constexpr FShaderSource::FViewType LineDirectiveEnd = ANSITEXTVIEW("\"\n");
 
 	static constexpr int32 HeaderLen =
 		DebugHashCommentStart.Len() +
diff --git a/Engine/Source/Runtime/RenderCore/Private/ShaderSource.cpp b/Engine/Source/Runtime/RenderCore/Private/ShaderSource.cpp
--- a/Engine/Source/Runtime/RenderCore/Private/ShaderSource.cpp
+++ b/Engine/Source/Runtime/RenderCore/Private/ShaderSource.cpp
@@ -32,15 +32,6 @@ void FShaderSource::Set(FShaderSource::FViewType InSrc, int32 AdditionalSlack)
 	FMemory::Memcpy(SourceData, InSrc.GetData(), sizeof(FShaderSource::CharType) * InSrc.Len());
 }
 
-#if !SHADER_SOURCE_ANSI
-void FShaderSource::Set(FAnsiStringView InSrc)
-{
-	TStringConvert<ANSICHAR, TCHAR> Convert;
-	SetLen(InSrc.Len());
-	Convert.Convert(Source.GetData(), Source.Num(), InSrc.GetData(), InSrc.Len());
-}
-#endif
-
 FShaderSource& FShaderSource::operator=(FShaderSource::FStringType&& InSrc)
 {
 	SourceCompressed.Empty();
diff --git a/Engine/Source/Runtime/RenderCore/Public/ShaderPreprocessTypes.h b/Engine/Source/Runtime/RenderCore/Public/ShaderPreprocessTypes.h
--- a/Engine/Source/Runtime/RenderCore/Public/ShaderPreprocessTypes.h
+++ b/Engine/Source/Runtime/RenderCore/Public/ShaderPreprocessTypes.h
@@ -73,35 +73,23 @@ public:
 
 	FAnsiStringView GetSourceViewAnsi() const
 	{
-#if SHADER_SOURCE_ANSI
 		return PreprocessedSource.GetView();
-#else
-		if (AnsiSource.IsEmpty())
-		{
-			AnsiSource = FAnsiString(PreprocessedSource.GetView());
-		}
-		return FAnsiStringView(AnsiSource);
-#endif
 	}
 
 	FStringView GetSourceViewWide() const
 	{
-#if SHADER_SOURCE_ANSI
 		if (WideSource.IsEmpty())
 		{
 			WideSource = FString(PreprocessedSource.GetView());
 		}
 		return FStringView(WideSource);
-#else
-		return PreprocessedSource.GetView();
-#endif
 	}
 
 	FStringView GetUnstrippedSourceView() const
 	{
 		// if the unstripped source is requested, check if the "original source" field has been populated
 		// if not then stripping hasn't occurred so there's only one preprocessed source; return it
-#if SHADER_SOURCE_ANSI
+
 		// convert and store wide versions of requested source if view is requested. 
 		// this is only used in debug paths (debug output and viewing source in-editor)
 		// and the job should be freed shortly after, so the memory overhead is not a concern
@@ -121,9 +109,6 @@ public:
 			}
 			return FStringView(WideSourceUnstripped);
 		}
-#else
-		return OriginalPreprocessedSource.IsEmpty() ? PreprocessedSource.GetView() : OriginalPreprocessedSource.GetView();
-#endif
 	}
 
 	FShaderSource& EditSource()
@@ -286,15 +271,10 @@ private:
 	// Set by Finalize; original preprocessed source as set by IShaderFormat::PreprocessShader
 	FShaderSource OriginalPreprocessedSource;
 
-#if SHADER_SOURCE_ANSI
-	// Set when GetSourceViewWide or GetUnstrippedSourceView accessors are called (if source is ANSI)
+	// Set when GetSourceViewWide or GetUnstrippedSourceView accessors are called
 	// Mutable so we can maintain const correctness in the API otherwise.
 	mutable FString WideSource;
 	mutable FString WideSourceUnstripped;
-#else
-	// Set when GetSourceViewAnsi is called (if source is wide)
-	mutable FAnsiString AnsiSource;
-#endif
 
 	// Array of errors encountered in preprocessing; should be populated by IShaderFormat::PreprocessShader
 	TArray<FShaderCompilerError> Errors;
diff --git a/Engine/Source/Runtime/RenderCore/Public/ShaderSource.h b/Engine/Source/Runtime/RenderCore/Public/ShaderSource.h
--- a/Engine/Source/Runtime/RenderCore/Public/ShaderSource.h
+++ b/Engine/Source/Runtime/RenderCore/Public/ShaderSource.h
@@ -8,15 +8,9 @@
 #include "Containers/UnrealString.h"
 #include "Misc/StringBuilder.h"
 
-#define SHADER_SOURCE_ANSI 1
-
-#if SHADER_SOURCE_ANSI
-#define SHADER_SOURCE_LITERAL(S) S
-#define SHADER_SOURCE_VIEWLITERAL(S) ANSITEXTVIEW(S)
-#else
-#define SHADER_SOURCE_LITERAL(S) TEXT(S)
-#define SHADER_SOURCE_VIEWLITERAL(S) TEXTVIEW(S)
-#endif
+#define SHADER_SOURCE_ANSI 1 UE_DEPRECATED_MACRO(5.6, "SHADER_SOURCE_ANSI has been deprecated and should be assumed to always be 1")
+#define SHADER_SOURCE_LITERAL(S) S UE_DEPRECATED_MACRO(5.6, "SHADER_SOURCE_LITERAL has been deprecated and should be replaced with just plain string literal")
+#define SHADER_SOURCE_VIEWLITERAL(S) ANSITEXTVIEW(S) UE_DEPRECATED_MACRO(5.6, "SHADER_SOURCE_VIEWLITERAL has been deprecated and should be replaced with ANSITEXTVIEW")
 
 /* Class used in shader compilation pipelines which wraps source code and ensures sufficient padding such that 16-byte wide SIMD
  * operations on the source are guaranteed to read valid memory even if starting from the last character.
@@ -34,7 +28,6 @@ private:
 	}
 
 public:
-#if SHADER_SOURCE_ANSI
 	typedef ANSICHAR CharType;
 	typedef FAnsiStringView FViewType;
 	typedef FAnsiString FStringType;
@@ -42,22 +35,16 @@ public:
 
 	template <int NumChars>
 	using TStringBuilder = TAnsiStringBuilder<NumChars>;
-#else
-	typedef TCHAR CharType;
-	typedef FStringView FViewType;
-	typedef FString FStringType;
-	typedef FCString FCStringType;
-
-	template <int NumChars>
-	using TStringBuilder = TStringBuilder<NumChars * 2>; // Buffer size for string builder is in bytes, not characters, hence *2
-#endif
 
 	/** Constexpr predicate indicating whether wide or ansi chars are used */
-	static constexpr bool IsWide() { return sizeof(CharType) == 2; }
+	UE_DEPRECATED(5.6, "Shader code is always assumed to be ANSI")
+	static constexpr bool IsWide() { return false; }
 	/** Constexpr function returning the number of characters read in a single SIMD compare op */
-	static constexpr int32 GetSimdCharCount() { return IsWide() ? 8 : 16; }
+	UE_DEPRECATED(5.6, "Shader code is always assumed to be ANSI, so GetSimdCharCount() is always 16")
+	static constexpr int32 GetSimdCharCount() { return 16; }
 	/** Constexpr function returning a mask value for a single character */
-	static constexpr int32 GetSingleCharMask() { return IsWide() ? 3 : 1; }
+	UE_DEPRECATED(5.6, "Shader code is always assumed to be ANSI, so GetSingleCharMask() is always 1")
+	static constexpr int32 GetSingleCharMask() { return 1; }
 
 	/* Construct an empty shader source object; will still contain padding */
 	FShaderSource() { SetLen(0); };
@@ -76,14 +63,6 @@ public:
 	 * @param AdditionalSlack optional additional space to allocate; this is on top of the automatic padding.
 	 */
 	RENDERCORE_API void Set(FViewType InSrc, int32 AdditionalSlack = 0);
-#if !SHADER_SOURCE_ANSI
-	/* Set the given ANSI string as the contents of this shader source object. The inner allocation will grow to fit
-	 * the string contents as needed. Note that this will incur a conversion to TCHAR and memcpy of the string contents.
-	 * Note that this is only needed if SHADER_SOURCE_ANSI is 0; redundant with Set function accepting a FViewType otherwise.
-	 * @param InSrc The source ANSI string to be copied
-	 */
-	RENDERCORE_API void Set(FAnsiStringView InSrc);
-#endif
 
 	/* Move assignment operator accepting a string object. This will append padding bytes to the existing string, as such it's
 	 * best if there's sufficient extra capacity in the string storage to avoid incurring a realloc-and-copy here.
@@ -158,11 +137,8 @@ public:
 	RENDERCORE_API void Decompress();
 	
 private:
-#if SHADER_SOURCE_ANSI
+	static_assert(sizeof(CharType) == 1);
 	static constexpr int32 ShaderSourceSimdPadding = 15;
-#else
-	static constexpr int32 ShaderSourceSimdPadding = 7;
-#endif
 	TArray<CharType> Source;
 	TArray<uint8> SourceCompressed;
 	int32 DecompressedCharCount = 0;
