[shaders] strip extensions from symbols info files & use map so entries are deduplicated on insertion. this allows platforms which export multiple artifacts with the same base filename (i.e. same hash) to avoid redundant entries in this file (PCD3D_SM6 will soon export both a .dxil and .pdb file as symbol artifacts as is required by Aftermath)

diff --git a/Engine/Source/Developer/ShaderCompilerCommon/Private/ShaderSymbolExport.cpp b/Engine/Source/Developer/ShaderCompilerCommon/Private/ShaderSymbolExport.cpp
--- a/Engine/Source/Developer/ShaderCompilerCommon/Private/ShaderSymbolExport.cpp
+++ b/Engine/Source/Developer/ShaderCompilerCommon/Private/ShaderSymbolExport.cpp
@@ -198,7 +198,8 @@ void FShaderSymbolExport::WriteSymbolData(const FString& Filename, const FString
 	if (ShouldGenerateShaderSymbolsInfo(ShaderFormat) && !DebugData.IsEmpty())
 	{
 		// Collect the simple shader symbol information
-		ShaderInfos.Add({ Filename, DebugData });
+		FString FilenameNoExt = FPaths::GetBaseFilename(Filename);
+		ShaderInfos.Add({ FilenameNoExt, DebugData });
 	}
 
 	if (ShouldWriteShaderSymbols(ShaderFormat) && Contents.Num())
@@ -296,13 +297,13 @@ void FShaderSymbolExport::NotifyShaderCompilersShutdown()
 			}
 
 			// sort and combine the data for output
-			ShaderInfos.Sort([](const FShaderInfo& A, const FShaderInfo& B) { return A.Hash < B.Hash; });
+			ShaderInfos.KeySort([](const FString& A, const FString& B) { return A < B; });
 
 			TArray<uint8> Output;
-			for (FShaderInfo Info : ShaderInfos)
+			for (TPair<FString, FString> Info : ShaderInfos)
 			{
-				auto TmpHash = StringCast<ANSICHAR>(*Info.Hash);
-				auto TmpData = StringCast<ANSICHAR>(*Info.Data);
+				auto TmpHash = StringCast<ANSICHAR>(*Info.Key);
+				auto TmpData = StringCast<ANSICHAR>(*Info.Value);
 				Output.Append((const uint8*)TmpHash.Get(), TmpHash.Length());
 				Output.Add(' ');
 				Output.Append((const uint8*)TmpData.Get(), TmpData.Length());
diff --git a/Engine/Source/Developer/ShaderCompilerCommon/Public/ShaderSymbolExport.h b/Engine/Source/Developer/ShaderCompilerCommon/Public/ShaderSymbolExport.h
--- a/Engine/Source/Developer/ShaderCompilerCommon/Public/ShaderSymbolExport.h
+++ b/Engine/Source/Developer/ShaderCompilerCommon/Public/ShaderSymbolExport.h
@@ -55,12 +55,7 @@ private:
 	bool bInitialized{ false };
 	bool bExportShaderSymbols{ false };
 
-	struct FShaderInfo
-	{
-		FString Hash;
-		FString Data;
-	};
-	TArray<FShaderInfo> ShaderInfos;
+	TMap<FString, FString> ShaderInfos;
 
 	uint32 DuplicateSymbols{ 0 };
 
