[Shaders] Fix interpolator removal for output-only system values.

There are system value interpolators that are only supported in the vertex output such as SV_ClipDistance and SV_CullDistance.
If they are delcared in the output struct between other interpolators, they occupy an interpolator register slot that will mismatch with the input interpolator in the pixel shader stage.
Moreover, this mismatch will only show up at PSO creation time, making it difficult to identify the source of the problem.
This problem is showing up more often now that more shader backends are making use of interpolator optimizations.

This change fixes the mismatch by moving all output-only system valus in destination struct to the bottom when unused interpolator removal is used.
This is safe because those output parameters can only be access in the vertex shader anyway, i.e. their register slot does not matter in the pixel shader,
but they would otherwise conflict with other interpolators that do matter in the pixel shader.

#rnx

diff --git a/Engine/Shaders/Public/ShaderVersion.ush b/Engine/Shaders/Public/ShaderVersion.ush
--- a/Engine/Shaders/Public/ShaderVersion.ush
+++ b/Engine/Shaders/Public/ShaderVersion.ush
@@ -3,4 +3,4 @@
 // Each time the console command is executed it generates a new GUID. As this file is included
 // in Platform.ush (which should be included in any shader) it allows to invalidate the shader DDC.
 // 
-#pragma message("UESHADERMETADATA_VERSION 4EA1C1A3-BE93-4931-8B40-7CFE5B6C2E85")
+#pragma message("UESHADERMETADATA_VERSION 7612E18B-B799-4D68-B8D1-84A56CFAED95")
diff --git a/Engine/Source/Developer/ShaderCompilerCommon/Private/HlslUtils.cpp b/Engine/Source/Developer/ShaderCompilerCommon/Private/HlslUtils.cpp
--- a/Engine/Source/Developer/ShaderCompilerCommon/Private/HlslUtils.cpp
+++ b/Engine/Source/Developer/ShaderCompilerCommon/Private/HlslUtils.cpp
@@ -754,9 +754,20 @@ struct FRemoveUnusedOutputs : FRemoveAlgorithm
 		return true;
 	}
 
-	bool IsSemanticUsed(const TCHAR* SemanticName)
+	// Returns true if the semantic names is explicitly exempt from removal by the shader backend. This applies to output-only system values such as SV_ClipDistance.
+	bool IsSemanticExempt(const TCHAR* SemanticName) const
+	{
+		return IsSubstringInArray(Exceptions, SemanticName);
+	}
+
+	bool IsSemanticExempt(const CrossCompiler::AST::FSemanticSpecifier* Semantic) const
+	{
+		return Semantic && IsSemanticExempt(Semantic->Semantic);
+	}
+
+	bool IsSemanticUsed(const TCHAR* SemanticName) const
 	{
-		if (bLeaveAllUsed || IsStringInArray(UsedOutputs, SemanticName) || IsSubstringInArray(Exceptions, SemanticName))
+		if (bLeaveAllUsed || IsStringInArray(UsedOutputs, SemanticName) || IsSemanticExempt(SemanticName))
 		{
 			return true;
 		}
@@ -772,7 +783,7 @@ struct FRemoveUnusedOutputs : FRemoveAlgorithm
 		return false;
 	}
 
-	bool IsSemanticUsed(const CrossCompiler::AST::FSemanticSpecifier* Semantic)
+	bool IsSemanticUsed(const CrossCompiler::AST::FSemanticSpecifier* Semantic) const
 	{
 		return Semantic ? IsSemanticUsed(Semantic->Semantic) : false;
 	}
@@ -898,6 +909,43 @@ struct FRemoveUnusedOutputs : FRemoveAlgorithm
 			}
 		}
 
+		/*
+		Move output only system values to the end, to ensure they don't occupy an output register that might be used by an input register in the pixel shader stage.
+		The following output struct for example will mismatch with input interpolators:
+
+			struct FOptimizedReturn
+			{
+				float4 Position : SV_POSITION;
+				float ClipDistance : SV_CLIPDISTANCE; // <-- Only available in vertex output
+				nointerpolation uint SomeIndex : SOMEINDEX;
+			};
+
+		The pixel shader will end up with the following input struct, because SV_CLIPDISTANCE is not supported in the pixel shader stage:
+
+			struct FOptimizedReturn
+			{
+				float4 Position : SV_POSITION;
+				nointerpolation uint SomeIndex : SOMEINDEX;
+			};
+
+		Now SV_CLIPDISTANCE will occupy output register 1 but this register is expected for "SOMEINDEX", resulting in a mismatch and PSO creation failure.
+		*/
+
+		auto IsMemberWithExemptSemantic = [this](FNode& Member) -> bool
+			{
+				CrossCompiler::AST::FDeclaratorList* MemberDeclList = Member.AsDeclaratorList();
+				checkf(MemberDeclList->Declarations.Num() == 1, TEXT("Destination structure for removed unused output interpolators must have exactly 1 member in its declaration list per statement, but %d are specified"), MemberDeclList->Declarations.Num());
+				CrossCompiler::AST::FDeclaration* MemberDecl = MemberDeclList->Declarations[0]->AsDeclaration();
+				return IsSemanticExempt(MemberDecl->Semantic);
+			};
+
+		DestStruct->Members.StableSort(
+			[this, &IsMemberWithExemptSemantic](CrossCompiler::AST::FNode& LhsNode, CrossCompiler::AST::FNode& RhsNode) -> bool
+			{
+				return IsMemberWithExemptSemantic(LhsNode) < IsMemberWithExemptSemantic(RhsNode);
+			}
+		);
+
 		return true;
 	}
 };
