Add a few basic unit tests for RemoveUnusedInputs

diff --git a/Engine/Source/Developer/ShaderCompilerCommon/Private/HlslUtils.cpp b/Engine/Source/Developer/ShaderCompilerCommon/Private/HlslUtils.cpp
--- a/Engine/Source/Developer/ShaderCompilerCommon/Private/HlslUtils.cpp
+++ b/Engine/Source/Developer/ShaderCompilerCommon/Private/HlslUtils.cpp
@@ -10,6 +10,10 @@
 #include "HlslParser.h"
 #include "ShaderCompilerCommon.h"
 
+#if WITH_LOW_LEVEL_TESTS
+#include "Tests/TestHarnessAdapter.h"
+#endif // WITH_LOW_LEVEL_TESTS
+
 static bool bLeaveAllUsed = false;
 
 namespace CrossCompiler
@@ -2142,3 +2146,199 @@ namespace UE::ShaderCompilerCommon
 		return bSuccess;
 	}
 }
+
+#if WITH_LOW_LEVEL_TESTS
+
+static TArray<FString> GetShaderInputSemantics(const FString& ShaderSource, const FString& EntryPoint)
+{
+	TArray<FString> Result;
+
+	CrossCompiler::FCompilerMessages CompilerMessages;
+	TArray<FScopedDeclarations> ScopedDeclarations;
+	CrossCompiler::Parser::Parse(ShaderSource, TEXT("shader.hlsl"), CompilerMessages, ScopedDeclarations,
+		[&Result, &EntryPoint](CrossCompiler::FLinearAllocator* Allocator, CrossCompiler::TLinearArray<CrossCompiler::AST::FNode*>& ASTNodes)
+		{
+			using namespace CrossCompiler::AST;
+
+			struct FStructSemantic
+			{
+				const TCHAR* StructName = nullptr;
+				const TCHAR* Semantic = nullptr;
+			};
+
+			TArray<FStructSemantic> KnownStructSemantics;
+
+			for (FNode* Node : ASTNodes)
+			{
+				if (FFunctionDefinition* FunctionDefinition = Node->AsFunctionDefinition())
+				{
+					FFunction* Prototype = FunctionDefinition->Prototype;
+
+					if (!Prototype || FCString::Strcmp(Prototype->Identifier, *EntryPoint))
+					{
+						continue;
+					}
+
+					for (FNode* ParamNode : Prototype->Parameters)
+					{
+						FParameterDeclarator* ParamDeclarator = ParamNode->AsParameterDeclarator();
+						if (!ParamDeclarator)
+						{
+							continue;
+						}
+
+						if (ParamDeclarator->Semantic && ParamDeclarator->Semantic->Semantic)
+						{
+							Result.Add(FString(ParamDeclarator->Semantic->Semantic));
+						}
+						else if (ParamDeclarator->Type && ParamDeclarator->Type->Specifier && ParamDeclarator->Type->Specifier->TypeName)
+						{
+							const TCHAR* TypeName = ParamDeclarator->Type->Specifier->TypeName;
+							for (const FStructSemantic& Semantic : KnownStructSemantics)
+							{
+								if (!FCString::Strcmp(Semantic.StructName, TypeName))
+								{
+									Result.Add(FString(Semantic.Semantic));
+								}
+							}
+						}
+					}
+				}
+				else if (FDeclaratorList* DeclaratorList = Node->AsDeclaratorList())
+				{
+					if (DeclaratorList->Type
+						&& DeclaratorList->Type->Specifier
+						&& DeclaratorList->Type->Specifier->Structure
+						&& DeclaratorList->Type->Specifier->Structure->Name)
+					{
+						auto Struct = DeclaratorList->Type->Specifier->Structure;
+						for (FNode* StructMember : Struct->Members)
+						{
+							FDeclaratorList* MemberDeclaratorList = StructMember->AsDeclaratorList();
+							if (!MemberDeclaratorList)
+							{
+								continue;
+							}
+
+							for (FNode* MemberDeclarationNode : MemberDeclaratorList->Declarations)
+							{
+								FDeclaration* MemberDeclaration = MemberDeclarationNode->AsDeclaration();
+								if (!MemberDeclaration)
+								{
+									continue;
+								}
+
+								if (MemberDeclaration->Semantic && MemberDeclaration->Semantic->Semantic)
+								{
+									FStructSemantic Entry;
+									Entry.StructName = DeclaratorList->Type->Specifier->Structure->Name;
+									Entry.Semantic = MemberDeclaration->Semantic->Semantic;
+									KnownStructSemantics.Add(Entry);
+								}
+							}
+						}
+					}
+				}
+			}
+		});
+
+	return Result;
+}
+
+TEST_CASE_NAMED(FShaderRemoveUnusedInputsLooseParametersTest, "Shaders::RemoveUnusedInputs::LooseParameters", "[EditorContext][EngineFilter]")
+{
+	FString ShaderCodeString(
+		TEXT(R"(
+float4 main(float4 A : TEXCOORD0, float4 B : TEXCOORD1) : SV_Target
+{
+	return B;
+}
+)"));
+
+	FString EntryPoint(TEXT("main"));
+
+	TArray<FString> ParametersToKeep;
+	ParametersToKeep.Add(TEXT("TEXCOORD1"));
+
+	TArray<FString> Errors;
+	const bool bSuccess = RemoveUnusedInputs(ShaderCodeString, ParametersToKeep, EntryPoint, Errors);
+	for (const FString& Message : Errors)
+	{
+		ADD_ERROR(Message);
+	}
+	TEST_TRUE(TEXT("RemoveUnusedInputs succeeded"), bSuccess);
+
+	TArray<FString> RewrittenInputs = GetShaderInputSemantics(ShaderCodeString, EntryPoint);
+	TEST_FALSE(TEXT("Rewritten shader uses TEXCOORD1"), RewrittenInputs.Contains(TEXT("TEXCOORD0")));
+	TEST_TRUE(TEXT("Rewritten shader uses TEXCOORD1"), RewrittenInputs.Contains(TEXT("TEXCOORD1")));
+}
+
+TEST_CASE_NAMED(FShaderRemoveUnusedInputsStructParametersTest, "Shaders::RemoveUnusedInputs::StructParameters", "[EditorContext][EngineFilter]")
+{
+	FString ShaderCodeString(
+		TEXT(R"(
+struct FPSInputs
+{
+	float4 A : TEXCOORD0;
+	float4 B : TEXCOORD1;
+};
+float4 main(FPSInputs Inputs) : SV_Target
+{
+	return Inputs.B;
+}
+)"));
+
+	FString EntryPoint(TEXT("main"));
+
+	TArray<FString> ParametersToKeep;
+	ParametersToKeep.Add(TEXT("TEXCOORD1"));
+
+	TArray<FString> Errors;
+	const bool bSuccess = RemoveUnusedInputs(ShaderCodeString, ParametersToKeep, EntryPoint, Errors);
+	for (const FString& Message : Errors)
+	{
+		ADD_ERROR(Message);
+	}
+	TEST_TRUE(TEXT("RemoveUnusedInputs succeeded"), bSuccess);
+
+	TArray<FString> RewrittenInputs = GetShaderInputSemantics(ShaderCodeString, EntryPoint);
+	TEST_FALSE(TEXT("Rewritten shader uses TEXCOORD0"), RewrittenInputs.Contains(TEXT("TEXCOORD0")));
+	TEST_TRUE(TEXT("Rewritten shader uses TEXCOORD1"), RewrittenInputs.Contains(TEXT("TEXCOORD1")));
+}
+
+TEST_CASE_NAMED(FShaderRemoveUnusedInputsMixedParametersTest, "Shaders::RemoveUnusedInputs::MixedParameters", "[EditorContext][EngineFilter]")
+{
+	FString ShaderCodeString(
+		TEXT(R"(
+struct FPSInputs
+{
+	float4 A : TEXCOORD0;
+	float4 B : TEXCOORD1;
+};
+float4 main(FPSInputs Inputs, float4 C : TEXCOORD2) : SV_Target
+{
+	return Inputs.A + C;
+}
+)"));
+
+	FString EntryPoint(TEXT("main"));
+
+	TArray<FString> ParametersToKeep;
+	ParametersToKeep.Add(TEXT("TEXCOORD1"));
+	ParametersToKeep.Add(TEXT("TEXCOORD2"));
+
+	TArray<FString> Errors;
+	const bool bSuccess = RemoveUnusedInputs(ShaderCodeString, ParametersToKeep, EntryPoint, Errors);
+	for (const FString& Message : Errors)
+	{
+		ADD_ERROR(Message);
+	}
+	TEST_TRUE(TEXT("RemoveUnusedInputs succeeded"), bSuccess);
+
+	TArray<FString> RewrittenInputs = GetShaderInputSemantics(ShaderCodeString, EntryPoint);
+	TEST_FALSE(TEXT("Rewritten shader uses TEXCOORD0"), RewrittenInputs.Contains(TEXT("TEXCOORD0")));
+	TEST_TRUE(TEXT("Rewritten shader uses TEXCOORD1"), RewrittenInputs.Contains(TEXT("TEXCOORD1")));
+	TEST_TRUE(TEXT("Rewritten shader uses TEXCOORD2"), RewrittenInputs.Contains(TEXT("TEXCOORD2")));
+}
+
+#endif // WITH_LOW_LEVEL_TESTS
