When removing unused shader inputs, inline the original entry point body into the optimized one

* This makes it cheaper for the native shader compiler to compile the final code as inlining during optimization is more expensive than pre-inlining

diff --git a/Engine/Shaders/Public/ShaderVersion.ush b/Engine/Shaders/Public/ShaderVersion.ush
--- a/Engine/Shaders/Public/ShaderVersion.ush
+++ b/Engine/Shaders/Public/ShaderVersion.ush
@@ -3,4 +3,4 @@
 // Each time the console command is executed it generates a new GUID. As this file is included
 // in Platform.ush (which should be included in any shader) it allows to invalidate the shader DDC.
 // 
-#pragma message("UESHADERMETADATA_VERSION FF5BEA66-2646-43E2-902C-736FC7521050")
+#pragma message("UESHADERMETADATA_VERSION F65F010C-86F8-422B-AD7C-D3895C89235A")
diff --git a/Engine/Source/Developer/ShaderCompilerCommon/Private/HlslUtils.cpp b/Engine/Source/Developer/ShaderCompilerCommon/Private/HlslUtils.cpp
--- a/Engine/Source/Developer/ShaderCompilerCommon/Private/HlslUtils.cpp
+++ b/Engine/Source/Developer/ShaderCompilerCommon/Private/HlslUtils.cpp
@@ -1049,14 +1049,7 @@ struct FRemoveUnusedInputs : FRemoveAlgorithm
 
 		FInputsBodyContext BodyContext;
 
-		// Setup the call to the original entry point
-		BodyContext.CallToOriginalFunction = new(Allocator) FFunctionExpression(Allocator, SourceInfo, MakeIdentifierExpression(Allocator, *EntryPoint, SourceInfo));
-
-		if (!SetupInputAndReturnType(EntryFunction, MiniSymbolTable, BodyContext))
-		{
-			bSuccess = false;
-			return;
-		}
+		BodyContext.NewInputStruct = CreateNewStructSpecifier(BodyContext.InputTypeName, BodyContext.NewStructs);
 
 		if (!ProcessOriginalParameters(EntryFunction, MiniSymbolTable, BodyContext))
 		{
@@ -1064,21 +1057,14 @@ struct FRemoveUnusedInputs : FRemoveAlgorithm
 			return;
 		}
 
-		// Real call statement
-		if (BodyContext.FinalInstruction)
-		{
-			auto* JumpStatement = BodyContext.FinalInstruction->AsJumpStatement();
-			if (JumpStatement)
-			{
-				JumpStatement->OptionalExpression = BodyContext.CallToOriginalFunction;
-			}
-		}
-		else
+		// Simply pre-inline the original function body.
+		// This is cheaper for the down-stream shader compilers as their own inlining passes are expensive when dealing with large functions.
+		if (EntryFunction->Body)
 		{
-			BodyContext.FinalInstruction = new(Allocator) CrossCompiler::AST::FExpressionStatement(Allocator, BodyContext.CallToOriginalFunction, SourceInfo);
+			BodyContext.PostInstructions.Append(EntryFunction->Body->Statements);
 		}
 
-		auto* Body = AddStatementsToBody(BodyContext, nullptr);
+		auto* Body = AddStatementsToBody(BodyContext, nullptr /*CallInstruction*/);
 
 		if (BodyContext.NewInputStruct->Members.Num() > 0)
 		{
@@ -1138,10 +1124,6 @@ struct FRemoveUnusedInputs : FRemoveAlgorithm
 			{
 				// Add this parameter as an input to the new function
 				BodyContext.NewFunctionParameters.Add(ParameterDeclarator);
-
-				// Add the parameter to the actual function call
-				auto* Parameter = MakeIdentifierExpression(Allocator, ParameterDeclarator->Identifier, SourceInfo);
-				BodyContext.CallToOriginalFunction->Expressions.Add(Parameter);
 			}
 		}
 		return true;
@@ -1157,17 +1139,14 @@ struct FRemoveUnusedInputs : FRemoveAlgorithm
 		auto* Initializer = new(Allocator) FUnaryExpression(Allocator, EOperators::TypeCast, Zero, SourceInfo);
 		Initializer->TypeSpecifier = MakeSimpleType(OriginalStructSpecifier->Name)->Specifier;
 
-		// Add a local variable to receive the output from the function
-		FString LocalStructVarName = TEXT("Local_");
-		LocalStructVarName += OriginalStructSpecifier->Name;
-		LocalStructVarName += TEXT("_IN");
-		auto* LocalStructVariable = CreateLocalVariable(OriginalStructSpecifier->Name, *LocalStructVarName, Initializer);
+		// Add a local variable to receive the output from the function.
+		// The name is simply the parameter name in the real entry point, since it will be inlined directly in the optimized version.
+
+		const TCHAR* LocalStructVarName = ParameterDeclarator->Identifier;
+		auto* LocalStructVariable = CreateLocalVariable(OriginalStructSpecifier->Name, LocalStructVarName, Initializer);
 		BodyContext.PreInstructions.Add(LocalStructVariable);
 
-		// Add the parameter to the actual function call
-		auto* Parameter = MakeIdentifierExpression(Allocator, *LocalStructVarName, SourceInfo);
-		BodyContext.CallToOriginalFunction->Expressions.Add(Parameter);
-		return AddUsedInputMembers(BodyContext.NewInputStruct, BodyContext.InputVariableName, OriginalStructSpecifier, *LocalStructVarName, MiniSymbolTable, BodyContext);
+		return AddUsedInputMembers(BodyContext.NewInputStruct, BodyContext.InputVariableName, OriginalStructSpecifier, LocalStructVarName, MiniSymbolTable, BodyContext);
 	}
 
 	bool IsSemanticUsed(const TCHAR* SemanticName)
@@ -1234,39 +1213,6 @@ struct FRemoveUnusedInputs : FRemoveAlgorithm
 
 		auto* LocalVar = CreateLocalVariable(ParameterDeclarator->Type->Specifier->TypeName, ParameterDeclarator->Identifier, Initializer);
 		BodyContext.PreInstructions.Add(LocalVar);
-
-		// Add the parameter to the actual function call
-		auto* Parameter = MakeIdentifierExpression(Allocator, ParameterDeclarator->Identifier, SourceInfo);
-		BodyContext.CallToOriginalFunction->Expressions.Add(Parameter);
-	}
-
-	bool SetupInputAndReturnType(CrossCompiler::AST::FFunctionDefinition* EntryFunction, TArray<CrossCompiler::AST::FStructSpecifier*>& MiniSymbolTable, FInputsBodyContext& BodyContext)
-	{
-		using namespace CrossCompiler::AST;
-
-		// New return type
-		BodyContext.NewInputStruct = CreateNewStructSpecifier(BodyContext.InputTypeName, BodyContext.NewStructs);
-
-		auto* ReturnType = EntryFunction->Prototype->ReturnType;
-		if (ReturnType && ReturnType->Specifier && ReturnType->Specifier->TypeName)
-		{
-			const TCHAR* ReturnTypeName = ReturnType->Specifier->TypeName;
-			if (!EntryFunction->Prototype->ReturnSemantic && !FCString::Strcmp(ReturnTypeName, TEXT("void")))
-			{
-				// No return instruction required
-			}
-			else
-			{
-				auto* ReturnStatement = new(Allocator) FJumpStatement(Allocator, EJumpType::Return, SourceInfo);
-				BodyContext.FinalInstruction = ReturnStatement;
-			}
-			return true;
-		}
-		else
-		{
-			Errors.Add(FString::Printf(TEXT("RemoveUnusedInputs: Internal error trying to determine return type")));
-		}
-		return false;
 	}
 
 	void WriteGeneratedInCode(CrossCompiler::AST::FFunctionDefinition* NewEntryFunction, TArray<CrossCompiler::AST::FStructSpecifier*>& NewStructs, FString& OutGeneratedCode)
