Unify multiple CrossCompiler::Parser::Parse() implementations and add FNodeContainer to hold AST nodes to be available outside of the Parse callback function

diff --git a/Engine/Source/Developer/ShaderCompilerCommon/Private/HlslParser.cpp b/Engine/Source/Developer/ShaderCompilerCommon/Private/HlslParser.cpp
--- a/Engine/Source/Developer/ShaderCompilerCommon/Private/HlslParser.cpp
+++ b/Engine/Source/Developer/ShaderCompilerCommon/Private/HlslParser.cpp
@@ -2402,64 +2402,49 @@ Done:
 	{
 		bool Parse(const FString& Input, const FString& Filename, FCompilerMessages& OutCompilerMessages, TConstArrayView<FScopedDeclarations> InScopedDeclarations, TCallback* Callback, void* CallbackData)
 		{
-			FLinearAllocator Allocator;
-			FHlslParser Parser(&Allocator, OutCompilerMessages, InScopedDeclarations);
-			if (!Parser.Scanner.Lex(Input, Filename))
-			{
-				return false;
-			}
+			CrossCompiler::FNodeContainer Container;
 
-			//Parser.Scanner.Dump();
+			const bool bSuccess = Parse(Input, Filename, OutCompilerMessages, InScopedDeclarations, Container);
 
-			bool bSuccess = true;
-			TLinearArray<AST::FNode*> Nodes(&Allocator);
-			while (Parser.Scanner.HasMoreTokens())
+			if (bSuccess && Callback)
 			{
-				auto LastIndex = Parser.Scanner.GetCurrentTokenIndex();
+				Callback(CallbackData, &Container.Allocator, Container.Nodes);
+			}
 
-				static FString GlobalDeclOrDefinition(TEXT("Global declaration or definition"));
-				AST::FNode* Node = nullptr;
-				auto Result = TryTranslationUnit(Parser, &Allocator, &Node);
-				if (Result == EParseResult::Error)
-				{
-					bSuccess = false;
-					break;
-				}
-				else
-				{
-					check(Result == EParseResult::Matched);
-					Nodes.Add(Node);
-				}
+			return bSuccess;
+		}
 
-				check(LastIndex != Parser.Scanner.GetCurrentTokenIndex());
-			}
+		bool Parse(const FString& Input, const FString& Filename, FCompilerMessages& OutCompilerMessages, TConstArrayView<FScopedDeclarations> InScopedDeclarations, TFunction< void(CrossCompiler::FLinearAllocator* Allocator, CrossCompiler::TLinearArray<CrossCompiler::AST::FNode*>& ASTNodes)> Function)
+		{
+			CrossCompiler::FNodeContainer Container;
 
-			if (bSuccess && Callback)
+			const bool bSuccess = Parse(Input, Filename, OutCompilerMessages, InScopedDeclarations, Container);
+
+			if (bSuccess)
 			{
-				Callback(CallbackData, &Allocator, Nodes);
+				Function(&Container.Allocator, Container.Nodes);
 			}
 
 			return bSuccess;
 		}
 
-		bool Parse(const FString& Input, const FString& Filename, FCompilerMessages& OutCompilerMessages, TConstArrayView<FScopedDeclarations> InScopedDeclarations, TFunction< void(CrossCompiler::FLinearAllocator* Allocator, CrossCompiler::TLinearArray<CrossCompiler::AST::FNode*>& ASTNodes)> Function)
+		bool Parse(const FString& Input, const FString& Filename, FCompilerMessages& OutCompilerMessages, TConstArrayView<FScopedDeclarations> InScopedDeclarations, CrossCompiler::FNodeContainer& OutContainer)
 		{
-			FLinearAllocator Allocator;
-			FHlslParser Parser(&Allocator, OutCompilerMessages, InScopedDeclarations);
+			FHlslParser Parser(&OutContainer.Allocator, OutCompilerMessages, InScopedDeclarations);
 			if (!Parser.Scanner.Lex(Input, Filename))
 			{
 				return false;
 			}
 
 			bool bSuccess = true;
-			TLinearArray<AST::FNode*> Nodes(&Allocator);
+
 			AST::FNode* PrevNode = nullptr;
 			while (Parser.Scanner.HasMoreTokens())
 			{
 				auto LastIndex = Parser.Scanner.GetCurrentTokenIndex();
 				static FString GlobalDeclOrDefinition(TEXT("Global declaration or definition"));
 				AST::FNode* Node = nullptr;
-				auto Result = TryTranslationUnit(Parser, &Allocator, &Node);
+				auto Result = TryTranslationUnit(Parser, &OutContainer.Allocator, &Node);
 				if (Result == EParseResult::Error)
 				{
 					bSuccess = false;
@@ -2472,27 +2457,16 @@ Done:
 				else
 				{
 					check(Result == EParseResult::Matched);
-					/*
-					if (bDump && Node)
-					{
-					Node->Dump(0);
-					}
-					*/
 					checkf(Node, TEXT("Null AST Node! PrevNode @ %s(%d)"),
 						((PrevNode && PrevNode->SourceInfo.Filename) ? *(*PrevNode->SourceInfo.Filename) : TEXT("<none>")),
 						(PrevNode ? PrevNode->SourceInfo.Line : 0));
-					Nodes.Add(Node);
+					OutContainer.Nodes.Add(Node);
 				}
 
 				check(LastIndex != Parser.Scanner.GetCurrentTokenIndex());
 				PrevNode = Node;
 			}
 
-			if (bSuccess)
-			{
-				Function(&Allocator, Nodes);
-			}
-
 			return bSuccess;
 		}
 
diff --git a/Engine/Source/Developer/ShaderCompilerCommon/Private/HlslParser.h b/Engine/Source/Developer/ShaderCompilerCommon/Private/HlslParser.h
--- a/Engine/Source/Developer/ShaderCompilerCommon/Private/HlslParser.h
+++ b/Engine/Source/Developer/ShaderCompilerCommon/Private/HlslParser.h
@@ -15,6 +15,8 @@ struct FScopedDeclarations;
 namespace CrossCompiler
 {
 	struct FLinearAllocator;
+	struct FNodeContainer;
+
 	namespace AST
 	{
 		class FNode;
@@ -38,6 +40,9 @@ namespace CrossCompiler
 		// Returns true if successfully parsed
 		SHADERCOMPILERCOMMON_API bool Parse(const FString& Input, const FString& Filename, FCompilerMessages& OutCompilerMessages, TConstArrayView<FScopedDeclarations> InScopedDeclarations, TFunction< void(CrossCompiler::FLinearAllocator* Allocator, CrossCompiler::TLinearArray<CrossCompiler::AST::FNode*>& ASTNodes)> Function);
 
+		// Returns true if successfully parsed
+		SHADERCOMPILERCOMMON_API bool Parse(const FString& Input, const FString& Filename, FCompilerMessages& OutCompilerMessages, TConstArrayView<FScopedDeclarations> InScopedDeclarations, CrossCompiler::FNodeContainer& OutContainer);
+
 		// Sample callback to write out all nodes into a string; pass a valid pointer FString* as OutFStringPointer
 		SHADERCOMPILERCOMMON_API void WriteNodesToString(void* OutFStringPointer, CrossCompiler::FLinearAllocator* Allocator, CrossCompiler::TLinearArray<CrossCompiler::AST::FNode*>& ASTNodes);
 	}
diff --git a/Engine/Source/Developer/ShaderCompilerCommon/Private/HlslUtils.h b/Engine/Source/Developer/ShaderCompilerCommon/Private/HlslUtils.h
--- a/Engine/Source/Developer/ShaderCompilerCommon/Private/HlslUtils.h
+++ b/Engine/Source/Developer/ShaderCompilerCommon/Private/HlslUtils.h
@@ -13,6 +13,11 @@
 
 namespace CrossCompiler
 {
+	namespace AST
+	{
+		class FNode;
+	}
+
 	namespace Memory
 	{
 		struct FPage
@@ -46,6 +51,8 @@ namespace CrossCompiler
 
 	struct FLinearAllocator
 	{
+		UE_NONCOPYABLE(FLinearAllocator);
+
 		FLinearAllocator()
 		{
 			auto* Initial = Memory::FPage::AllocatePage(Memory::MinPageSize);
@@ -213,6 +220,19 @@ namespace CrossCompiler
 		}
 	};
 
+	struct FNodeContainer
+	{
+		UE_NONCOPYABLE(FNodeContainer);
+
+		FNodeContainer()
+			: Nodes(&Allocator)
+		{
+		}
+
+		FLinearAllocator Allocator;
+		TLinearArray<CrossCompiler::AST::FNode*> Nodes;
+	};
+
 	/*
 	template <typename TType>
 	struct TLinearSet : public TSet<typename TType, DefaultKeyFuncs<typename TType>, FLinearSetAllocator>
