Groundwork for the ability to automatically rename identifiers in shader functions. This is required to implement hygienic macros / funciton pre-inlining which aims to improve shader compilation speed.

* Add FIdentifier type and use it instead of raw TCHAR* in the AST nodes. This adds an indirection for string data, which allows for efficient renaming, where only the declaration nodes need to be updated.
* Keep track and deduplicate local identifiers inside FSymbolScope during parsing.

diff --git a/Engine/Source/Developer/ShaderCompilerCommon/Private/HlslAST.cpp b/Engine/Source/Developer/ShaderCompilerCommon/Private/HlslAST.cpp
--- a/Engine/Source/Developer/ShaderCompilerCommon/Private/HlslAST.cpp
+++ b/Engine/Source/Developer/ShaderCompilerCommon/Private/HlslAST.cpp
@@ -329,7 +329,7 @@ namespace CrossCompiler
 			if (IsConstant())
 			{
 				checkf(Identifier!=nullptr, TEXT("Null identifier, literaltype %d"), (int32)LiteralType);
-				OutValue = (int32)FCString::Atoi(Identifier);
+				OutValue = (int32)FCString::Atoi(Identifier->GetData());
 				return true;
 			}
 
diff --git a/Engine/Source/Developer/ShaderCompilerCommon/Private/HlslAST.h b/Engine/Source/Developer/ShaderCompilerCommon/Private/HlslAST.h
--- a/Engine/Source/Developer/ShaderCompilerCommon/Private/HlslAST.h
+++ b/Engine/Source/Developer/ShaderCompilerCommon/Private/HlslAST.h
@@ -29,6 +29,76 @@ namespace CrossCompiler
 		struct FExpressionStatement;
 		struct FSemanticSpecifier;
 
+		// Wrapper for identifier strings.
+		// Provides an indirection that can be used for efficient symbol renaming, etc.
+		class FIdentifier
+		{
+		private:
+			FIdentifier(const TCHAR* InString, int32 InLength)
+				: StringView(InString, InLength)
+			{
+			}
+			~FIdentifier()
+			{
+				// We don't expect this to ever get destroyed, as we only allocate everything from the linear allocator
+				checkNoEntry();
+			}
+
+		public:
+
+			// NOTE: All memory is owned by the liner allocator and is released when it is destroyed. No need to delete it explicitly.
+			static FIdentifier* New(FLinearAllocator* Allocator, const TCHAR* InString, int32 InLength = -1)
+			{
+				if (InLength < 0)
+				{
+					InLength = FCString::Strlen(InString);
+				}
+				const TCHAR* DuplicatedString = Allocator->Strdup(InString, InLength);
+				void* IdentifierMemory = Allocator->Alloc(sizeof(FIdentifier), alignof(FIdentifier));
+				return new(IdentifierMemory) FIdentifier(DuplicatedString, InLength);
+			}
+
+			static FIdentifier* New(FLinearAllocator* Allocator, FStringView InString)
+			{
+				return FIdentifier::New(Allocator, InString.GetData(), InString.Len());
+			}
+
+			FStringView ToStringView() const
+			{
+				return StringView;
+			};
+
+			FString ToString() const
+			{
+				return FString(StringView);
+			};
+
+			const TCHAR* GetData() const
+			{
+				return StringView.GetData();
+			}
+
+			int32 Len() const
+			{
+				return StringView.Len();
+			}
+
+			bool Equals(FStringView Other) const
+			{
+				return Other.Equals(StringView, ESearchCase::CaseSensitive);
+			}
+
+			void Rename(FLinearAllocator* Allocator, FStringView InString)
+			{
+				int32 NewLength = InString.Len();
+				const TCHAR* NewData = Allocator->Strdup(InString.GetData(), InString.Len());
+				StringView = FStringView(NewData, NewLength);
+			}
+
+		private:
+			FStringView StringView;
+		};
+
 		struct FASTWriter
 		{
 			FString& Output;
@@ -61,6 +131,16 @@ namespace CrossCompiler
 				return *this;
 			}
 
+			inline FASTWriter& operator << (const FIdentifier* Identifier)
+			{
+				if (Identifier)
+				{
+					Output += Identifier->ToStringView();
+				}
+
+				return *this;
+			}
+
 			inline FASTWriter& operator << (const TCHAR Char)
 			{
 				if (Char != 0)
@@ -407,7 +487,7 @@ namespace CrossCompiler
 				ELiteralType LiteralType;
 				struct FTypeSpecifier* TypeSpecifier;
 			};
-			const TCHAR* Identifier;
+			const FIdentifier* Identifier;
 
 			TLinearArray<FExpression*> Expressions;
 
@@ -524,7 +604,7 @@ namespace CrossCompiler
 			virtual void Write(FASTWriter& Writer) const override;
 			virtual FDeclaration* AsDeclaration() override { return this; }
 
-			const TCHAR* Identifier;
+			FIdentifier* Identifier; // non-const to allow renaming
 
 			FSemanticSpecifier* Semantic;
 
@@ -651,7 +731,7 @@ namespace CrossCompiler
 			virtual FParameterDeclarator* AsParameterDeclarator() override { return this; }
 
 			FFullySpecifiedType* Type;
-			const TCHAR* Identifier;
+			FIdentifier* Identifier; // non-const to allow renaming
 			FSemanticSpecifier* Semantic;
 			bool bIsArray;
 
@@ -667,8 +747,8 @@ namespace CrossCompiler
 			virtual void Write(FASTWriter& Writer) const override;
 
 			FFullySpecifiedType* ReturnType;
-			const TCHAR* ScopeIdentifier;
-			const TCHAR* Identifier;
+			const FIdentifier* ScopeIdentifier;
+			const FIdentifier* Identifier;
 			FSemanticSpecifier* ReturnSemantic;
 
 			TLinearArray<FNode*> Parameters;
diff --git a/Engine/Source/Developer/ShaderCompilerCommon/Private/HlslExpressionParser.inl b/Engine/Source/Developer/ShaderCompilerCommon/Private/HlslExpressionParser.inl
--- a/Engine/Source/Developer/ShaderCompilerCommon/Private/HlslExpressionParser.inl
+++ b/Engine/Source/Developer/ShaderCompilerCommon/Private/HlslExpressionParser.inl
@@ -11,6 +11,7 @@
 #include "Containers/Set.h"
 #include "Developer/ShaderCompilerCommon/Private/HlslParser.h"
 #include "Developer/ShaderCompilerCommon/Private/HlslAST.h"
+#include "Hash/xxhash.h"
 
 class Error;
 
@@ -26,35 +27,87 @@ namespace CrossCompiler
 
 	struct FSymbolScope
 	{
+		FLinearAllocator* Allocator;
 		FSymbolScope* Parent;
 		const TCHAR* Name;
 
 		TSet/*TLinearSet*/<FString> Symbols;
 		TLinearArray<FSymbolScope> Children;
 
-		FSymbolScope(FLinearAllocator* InAllocator, FSymbolScope* InParent) :
-			Parent(InParent),
-			Name(nullptr),
-			Children(InAllocator)
+		struct FIdentifierKeyFuncs : public DefaultKeyFuncs<FStringView>
 		{
+			static FStringView GetSetKey(FStringView K) { return K; }
+			template <typename T>
+			static FStringView GetSetKey(const TPair<FStringView, T>& P) { return P.Key; }
+			static bool Matches(FStringView A, FStringView B) { return A.Equals(B, ESearchCase::CaseSensitive); }
+			static uint32 GetKeyHash(FStringView Key)
+			{
+				return FXxHash64::HashBuffer(Key.GetData(), Key.Len() * sizeof(*Key.GetData())).Hash;
+			}
+		};
+
+		TMap<FStringView, CrossCompiler::AST::FIdentifier*, FDefaultSetAllocator, FIdentifierKeyFuncs> Identifiers;
+
+		FSymbolScope(FLinearAllocator* InAllocator, FSymbolScope* InParent) 
+			: Allocator(InAllocator)
+			, Parent(InParent)
+			, Name(nullptr)
+			, Children(InAllocator)
+		{
+		}
+		~FSymbolScope() 
+		{
+		
 		}
-		~FSymbolScope() {}
 
-		void Add(const FString& Type)
+		FORCENOINLINE void Add(const FString& Type)
 		{
 			Symbols.Emplace(Type);
 		}
 
-		void Add(FStringView Type)
+		FORCENOINLINE void Add(FStringView Type)
 		{
 			Symbols.Emplace(Type);
 		}
 
-		void Add(const TCHAR* Type)
+		FORCENOINLINE void Add(const TCHAR* Type)
 		{
 			Symbols.Emplace(Type);
 		}
 
+		CrossCompiler::AST::FIdentifier* FindExistingIdentifier(uint32 IdentifierHash, FStringView Identifier, bool bSearchUpwards)
+		{
+			if (auto ExistingEntry = Identifiers.FindByHash(IdentifierHash, Identifier))
+			{
+				return *ExistingEntry;
+			}
+			else if (bSearchUpwards && Parent)
+			{
+				return Parent->FindExistingIdentifier(IdentifierHash, Identifier, bSearchUpwards);
+			}
+			else
+			{
+				return nullptr;
+			}
+		}
+
+		CrossCompiler::AST::FIdentifier* FindOrAddIdentifier(FStringView IdentifierName)
+		{
+			const uint32 IdentifierHash = FIdentifierKeyFuncs::GetKeyHash(IdentifierName);
+
+			const bool bSearchUpwards = true;
+			if (CrossCompiler::AST::FIdentifier* ExistingEntry = FindExistingIdentifier(IdentifierHash, IdentifierName, bSearchUpwards))
+			{
+				return ExistingEntry;
+			}
+			else
+			{
+				AST::FIdentifier* NewIdentifier = AST::FIdentifier::New(Allocator, IdentifierName);
+				Identifiers.AddByHash(IdentifierHash, IdentifierName, NewIdentifier);
+				return NewIdentifier;
+			}
+		}
+
 		static bool FindType(const FSymbolScope* Scope, const FString& Type, bool bSearchUpwards = true)
 		{
 			while (Scope)
@@ -105,6 +158,18 @@ namespace CrossCompiler
 		}
 	};
 
+	inline AST::FIdentifier* FindOrAddIdentifier(FLinearAllocator* Allocator, FSymbolScope* SymbolScope, FStringView IdentifierName)
+	{
+		if (SymbolScope)
+		{
+			return SymbolScope->FindOrAddIdentifier(IdentifierName);
+		}
+		else
+		{
+			return AST::FIdentifier::New(Allocator, IdentifierName);
+		}
+	}
+
 	struct FCreateSymbolScope
 	{
 		FSymbolScope* Original;
@@ -819,7 +884,7 @@ namespace CrossCompiler
 					return ParseResultError();
 				}
 				auto* FieldExpression = new(Allocator) AST::FUnaryExpression(Allocator, AST::EOperators::FieldSelection, PrevExpression, Token->SourceInfo);
-				FieldExpression->Identifier = Allocator->Strdup(Identifier->String);
+				FieldExpression->Identifier = AST::FIdentifier::New(Allocator, Identifier->String);
 				PrevExpression = FieldExpression;
 				bFoundAny = true;
 			}
@@ -905,7 +970,7 @@ namespace CrossCompiler
 			return ParseResultError();
 		}
 
-		auto ParseScopedIdentifier = [](FHlslScanner& Scanner, FLinearAllocator* Allocator, FString& Name, const FHlslToken* Token, const FHlslToken* Token1, const FHlslToken* Token2)
+		auto ParseScopedIdentifier = [SymbolScope](FHlslScanner& Scanner, FLinearAllocator* Allocator, FString& Name, const FHlslToken* Token, const FHlslToken* Token1, const FHlslToken* Token2)
 		{
 			while (Token1 && Token1->Token == EHlslToken::ColonColon && Token2 && Token2->Token == EHlslToken::Identifier)
 			{
@@ -918,7 +983,7 @@ namespace CrossCompiler
 			}
 
 			AST::FExpression* AtomExpression = new(Allocator) AST::FExpression(Allocator, AST::EOperators::Identifier, Token->SourceInfo);
-			AtomExpression->Identifier = Allocator->Strdup(*Name);
+			AtomExpression->Identifier = FindOrAddIdentifier(Allocator, SymbolScope, *Name);
 			return AtomExpression;
 		};
 
@@ -928,7 +993,7 @@ namespace CrossCompiler
 		case EHlslToken::Literal:
 			Scanner.Advance();
 			AtomExpression = new(Allocator) AST::FExpression(Allocator, AST::EOperators::Literal, Token->SourceInfo);
-			AtomExpression->Identifier = Allocator->Strdup(Token->String);
+			AtomExpression->Identifier = AST::FIdentifier::New(Allocator, Token->String);
 			AtomExpression->LiteralType = Token->LiteralType;
 			break;
 
@@ -962,7 +1027,7 @@ namespace CrossCompiler
 			{
 				Scanner.Advance();
 				AtomExpression = new(Allocator) AST::FExpression(Allocator, AST::EOperators::Identifier, Token->SourceInfo);
-				AtomExpression->Identifier = Allocator->Strdup(Token->String);
+				AtomExpression->Identifier = FindOrAddIdentifier(Allocator, SymbolScope, Token->String);
 			}
 		}
 			break;
@@ -1014,7 +1079,7 @@ namespace CrossCompiler
 			{
 				//@todo-rco: Check this var exists on the symbol table
 				AtomExpression = new(Allocator) AST::FExpression(Allocator, AST::EOperators::Identifier, TypeSpecifier->SourceInfo);
-				AtomExpression->Identifier = TypeSpecifier->TypeName;
+				AtomExpression->Identifier = AST::FIdentifier::New(Allocator, TypeSpecifier->TypeName);
 				break;
 			}
 			// Handle float3(x,y,z)
@@ -1023,7 +1088,7 @@ namespace CrossCompiler
 				if (Scanner.MatchToken(EHlslToken::LeftParenthesis))
 				{
 					auto* TypeExpression = new(Allocator) AST::FExpression(Allocator, AST::EOperators::Identifier, TypeSpecifier->SourceInfo);
-					TypeExpression->Identifier = TypeSpecifier->TypeName;
+					TypeExpression->Identifier = AST::FIdentifier::New(Allocator, TypeSpecifier->TypeName);
 					auto* FunctionCall = new(Allocator) AST::FFunctionExpression(Allocator, Token->SourceInfo, TypeExpression);
 					auto Result = ParseExpressionList(EHlslToken::RightParenthesis, Scanner, SymbolScope, EHlslToken::Invalid, Allocator, FunctionCall);
 					if (Result != EParseResult::Matched)
diff --git a/Engine/Source/Developer/ShaderCompilerCommon/Private/HlslParser.cpp b/Engine/Source/Developer/ShaderCompilerCommon/Private/HlslParser.cpp
--- a/Engine/Source/Developer/ShaderCompilerCommon/Private/HlslParser.cpp
+++ b/Engine/Source/Developer/ShaderCompilerCommon/Private/HlslParser.cpp
@@ -306,12 +306,12 @@ namespace CrossCompiler
 				if (ParseGeneralType(Scanner, ETF_SAMPLER_TEXTURE_BUFFER, nullptr, Allocator, &DummyTypeSpecifier) == EParseResult::Matched)
 				{
 					Declaration = new(Allocator) AST::FDeclaration(Allocator, DummyTypeSpecifier->SourceInfo);
-					Declaration->Identifier = Allocator->Strdup(DummyTypeSpecifier->TypeName);
+					Declaration->Identifier = AST::FIdentifier::New(Allocator, DummyTypeSpecifier->TypeName);
 				}
 				else if (Scanner.MatchToken(EHlslToken::Identifier))
 				{
 					Declaration = new(Allocator) AST::FDeclaration(Allocator, IdentifierToken->SourceInfo);
-					Declaration->Identifier = Allocator->Strdup(IdentifierToken->String);
+					Declaration->Identifier = AST::FIdentifier::New(Allocator, IdentifierToken->String);
 				}
 				else
 				{
@@ -684,7 +684,7 @@ namespace CrossCompiler
 						FullType->Specifier = TypeSpecifier;
 
 						auto* Declaration = new(Allocator)AST::FDeclaration(Allocator, IdentifierToken->SourceInfo);
-						Declaration->Identifier = Allocator->Strdup(IdentifierToken->String);
+						Declaration->Identifier = FindOrAddIdentifier(Allocator, SymbolScope, IdentifierToken->String);
 
 						DeclaratorList->Declarations.Add(Declaration);
 						*OutDeclaratorList = DeclaratorList;
@@ -737,7 +737,7 @@ namespace CrossCompiler
 					//... Instance
 
 					auto* Declaration = new(Allocator) AST::FDeclaration(Allocator, IdentifierToken->SourceInfo);
-					Declaration->Identifier = Allocator->Strdup(IdentifierToken->String);
+					Declaration->Identifier = FindOrAddIdentifier(Allocator, SymbolScope, IdentifierToken->String);
 
 					if (ParseDeclarationMultiArrayBracketsAndIndex(Parser.Scanner, SymbolScope, false, Allocator, Declaration) == EParseResult::Error)
 					{
@@ -782,7 +782,7 @@ namespace CrossCompiler
 					}
 
 					auto* Declaration = new(Allocator) AST::FDeclaration(Allocator, IdentifierToken->SourceInfo);
-					Declaration->Identifier = Allocator->Strdup(IdentifierToken->String);
+					Declaration->Identifier = FindOrAddIdentifier(Allocator, SymbolScope, IdentifierToken->String);
 					//AddVar
 
 					if (ParseDeclarationMultiArrayBracketsAndIndex(Parser.Scanner, SymbolScope, false, Allocator, Declaration) == EParseResult::Error)
@@ -1253,16 +1253,16 @@ namespace CrossCompiler
 		auto* Function = new(Allocator) AST::FFunction(Allocator, Identifier->SourceInfo);
 		if (OperatorIdentifier)
 		{
-			Function->Identifier = OperatorIdentifier;
+			Function->Identifier = AST::FIdentifier::New(Allocator, OperatorIdentifier);
 			Function->bIsOperator = true;
 		}
 		else
 		{
-			Function->Identifier = Allocator->Strdup(Identifier->String);
+			Function->Identifier = FindOrAddIdentifier(Allocator, Parser.CurrentScope, Identifier->String);
 		}
 		if (ScopeIdentifier)
 		{
-			Function->ScopeIdentifier = Allocator->Strdup(ScopeIdentifier->String);
+			Function->ScopeIdentifier = AST::FIdentifier::New(Allocator, ScopeIdentifier->String);
 		}
 		Function->ReturnType = ReturnType;
 
@@ -1488,7 +1488,7 @@ Done:
 				// Adding function declarations to symbols defined within specific namespace
 				if (Parser.CurrentScope && Parser.CurrentScope->Name != nullptr)
 				{
-					Parser.CurrentScope->Symbols.Add(FunctionDefinition->Prototype->Identifier);
+					Parser.CurrentScope->Symbols.Add(FunctionDefinition->Prototype->Identifier->ToString());
 				}
 			}
 		}
@@ -2210,9 +2210,9 @@ Done:
 						CrossCompiler::AST::FDeclaration* Decl = DeclNode->AsDeclaration();
 						if (Decl)
 						{
-							if (Decl->Identifier && Decl->Identifier[0])
+							if (Decl->Identifier && Decl->Identifier->Len())
 							{
-								Parser.CurrentScope->Add(Decl->Identifier);
+								Parser.CurrentScope->Add(Decl->Identifier->GetData());
 							}
 						}
 						else
diff --git a/Engine/Source/Developer/ShaderCompilerCommon/Private/HlslUtils.cpp b/Engine/Source/Developer/ShaderCompilerCommon/Private/HlslUtils.cpp
--- a/Engine/Source/Developer/ShaderCompilerCommon/Private/HlslUtils.cpp
+++ b/Engine/Source/Developer/ShaderCompilerCommon/Private/HlslUtils.cpp
@@ -181,14 +181,19 @@ struct FRemoveAlgorithm
 	{
 	}
 
-	static CrossCompiler::AST::FExpression* MakeIdentifierExpression(CrossCompiler::FLinearAllocator* Allocator, const TCHAR* Name, const CrossCompiler::FSourceInfo& SourceInfo)
+	static CrossCompiler::AST::FExpression* MakeIdentifierExpression(CrossCompiler::FLinearAllocator* Allocator, const CrossCompiler::AST::FIdentifier* Name, const CrossCompiler::FSourceInfo& SourceInfo)
 	{
 		using namespace CrossCompiler::AST;
 		FExpression* Expression = new(Allocator) FExpression(Allocator, EOperators::Identifier, SourceInfo);
-		Expression->Identifier = Allocator->Strdup(Name);
+		Expression->Identifier = Name;
 		return Expression;
 	}
 
+	static CrossCompiler::AST::FExpression* MakeIdentifierExpression(CrossCompiler::FLinearAllocator* Allocator, const TCHAR* Name, const CrossCompiler::FSourceInfo& SourceInfo)
+	{
+		return MakeIdentifierExpression(Allocator, CrossCompiler::AST::FIdentifier::New(Allocator, Name), SourceInfo);
+	}
+
 	CrossCompiler::AST::FFunctionDefinition* FindEntryPointAndPopulateSymbolTable(CrossCompiler::TLinearArray<CrossCompiler::AST::FNode*>& ASTNodes, TArray<CrossCompiler::AST::FStructSpecifier*>& OutMiniSymbolTable, FString* OutOptionalWriteNodes)
 	{
 		using namespace CrossCompiler::AST;
@@ -206,7 +211,7 @@ struct FRemoveAlgorithm
 			}
 			else if (FFunctionDefinition* FunctionDefinition = Node->AsFunctionDefinition())
 			{
-				if (FCString::Strcmp(*EntryPoint, FunctionDefinition->Prototype->Identifier) == 0)
+				if (FunctionDefinition->Prototype->Identifier->Equals(EntryPoint))
 				{
 					EntryFunction = FunctionDefinition;
 				}
@@ -252,7 +257,7 @@ struct FRemoveAlgorithm
 		// New Entry definition/prototype
 		FFunctionDefinition* NewEntryFunction = new(Allocator) FFunctionDefinition(Allocator, SourceInfo);
 		NewEntryFunction->Prototype = new(Allocator) FFunction(Allocator, SourceInfo);
-		NewEntryFunction->Prototype->Identifier = Allocator->Strdup(*(EntryPoint + TEXT("__OPTIMIZED")));
+		NewEntryFunction->Prototype->Identifier = FIdentifier::New(Allocator, *(EntryPoint + TEXT("__OPTIMIZED")));
 		NewEntryFunction->Prototype->ReturnType = ReturnType;
 		NewEntryFunction->Prototype->Attributes = FunctionAttributes;
 		NewEntryFunction->Body = Body;
@@ -319,10 +324,10 @@ struct FRemoveAlgorithm
 		// Add copy statement(s)
 		FString LHSName = DestPrefix;
 		LHSName += '.';
-		LHSName += Declaration->Identifier;
+		LHSName += Declaration->Identifier->ToStringView();
 		FString RHSName = SourcePrefix;
 		RHSName += '.';
-		RHSName += Declaration->Identifier;
+		RHSName += Declaration->Identifier->ToStringView();
 
 		if (Declaration->bIsArray)
 		{
@@ -353,18 +358,23 @@ struct FRemoveAlgorithm
 		return true;
 	}
 
-	CrossCompiler::AST::FDeclaratorList* CreateLocalVariable(const TCHAR* Type, const TCHAR* Name, CrossCompiler::AST::FExpression* Initializer = nullptr)
+	CrossCompiler::AST::FDeclaratorList* CreateLocalVariable(const TCHAR* Type, CrossCompiler::AST::FIdentifier* VariableName, CrossCompiler::AST::FExpression* Initializer = nullptr)
 	{
 		using namespace CrossCompiler::AST;
 		auto* LocalVarDeclaratorList = new(Allocator) FDeclaratorList(Allocator, SourceInfo);
 		LocalVarDeclaratorList->Type = MakeSimpleType(Type);
 		auto* LocalVarDeclaration = new(Allocator) FDeclaration(Allocator, SourceInfo);
-		LocalVarDeclaration->Identifier = Allocator->Strdup(Name);
+		LocalVarDeclaration->Identifier = VariableName;
 		LocalVarDeclaration->Initializer = Initializer;
 		LocalVarDeclaratorList->Declarations.Add(LocalVarDeclaration);
 		return LocalVarDeclaratorList;
 	}
 
+	CrossCompiler::AST::FDeclaratorList* CreateLocalVariable( const TCHAR* Type, FStringView VariableName, CrossCompiler::AST::FExpression* Initializer = nullptr)
+	{
+		return CreateLocalVariable(Type, CrossCompiler::AST::FIdentifier::New(Allocator, VariableName), Initializer);
+	}
+
 	CrossCompiler::AST::FCompoundStatement* AddStatementsToBody(FBodyContext& Return, CrossCompiler::AST::FNode* CallInstruction)
 	{
 		CrossCompiler::AST::FCompoundStatement* Body = new(Allocator)CrossCompiler::AST::FCompoundStatement(Allocator, SourceInfo);
@@ -397,14 +407,14 @@ struct FRemoveAlgorithm
 		using namespace CrossCompiler::AST;
 		if (!A->bIsArray)
 		{
-			Errors.Add(FString::Printf(TEXT("RemoveUnusedOutputs: %s is expected to be an array!"), A->Identifier));
+			Errors.Add(FString::Printf(TEXT("RemoveUnusedOutputs: %s is expected to be an array!"), A->Identifier->GetData()));
 			return false;
 		}
 		else
 		{
 			if (A->ArraySize.Num() > 1)
 			{
-				Errors.Add(FString::Printf(TEXT("RemoveUnusedOutputs: No support for multidimensional arrays on %s!"), A->Identifier));
+				Errors.Add(FString::Printf(TEXT("RemoveUnusedOutputs: No support for multidimensional arrays on %s!"), A->Identifier->GetData()));
 				return false;
 			}
 
@@ -413,7 +423,7 @@ struct FRemoveAlgorithm
 				int32 DimA = 0;
 				if (!A->ArraySize[Index]->GetConstantIntValue(DimA))
 				{
-					Errors.Add(FString::Printf(TEXT("RemoveUnusedOutputs: Array %s is not a compile-time constant expression!"), A->Identifier));
+					Errors.Add(FString::Printf(TEXT("RemoveUnusedOutputs: Array %s is not a compile-time constant expression!"), A->Identifier->GetData()));
 					return false;
 				}
 
@@ -562,7 +572,7 @@ struct FRemoveUnusedOutputs : FRemoveAlgorithm
 
 		FCompoundStatement* Body = AddStatementsToBody(BodyContext, CallInstruction);
 		FFunctionDefinition* NewEntryFunction = CreateNewEntryFunction(Body, MakeSimpleType(BodyContext.NewReturnStruct->Name), BodyContext.NewFunctionParameters, EntryFunction->Prototype->Attributes);
-		EntryPoint = NewEntryFunction->Prototype->Identifier;
+		EntryPoint = NewEntryFunction->Prototype->Identifier->ToString();
 		WriteGeneratedOutCode(NewEntryFunction, BodyContext.NewStructs, GeneratedCode);
 		bSuccess = true;
 	}
@@ -571,7 +581,7 @@ struct FRemoveUnusedOutputs : FRemoveAlgorithm
 	{
 		CrossCompiler::AST::FASTWriter Writer(OutGeneratedCode);
 		GeneratedCode = TEXT("#line 1 \"RemoveUnusedOutputs.usf\"\n// Generated Entry Point: ");
-		GeneratedCode += NewEntryFunction->Prototype->Identifier;
+		GeneratedCode += NewEntryFunction->Prototype->Identifier->ToStringView();
 		GeneratedCode += TEXT("\n");
 		if (UsedOutputs.Num() > 0)
 		{
@@ -627,7 +637,7 @@ struct FRemoveUnusedOutputs : FRemoveAlgorithm
 			// Add the parameter to the actual function call
 			FString ParameterName = BodyContext.ReturnVariableName;
 			ParameterName += TEXT(".");
-			ParameterName += ParameterDeclarator->Identifier;
+			ParameterName += ParameterDeclarator->Identifier->ToStringView();
 			auto* Parameter = MakeIdentifierExpression(Allocator, *ParameterName, SourceInfo);
 			BodyContext.CallToOriginalFunction->Expressions.Add(Parameter);
 		}
@@ -651,7 +661,7 @@ struct FRemoveUnusedOutputs : FRemoveAlgorithm
 		auto* MemberDeclaratorList = new(Allocator) FDeclaratorList(Allocator, SourceInfo);
 		MemberDeclaratorList->Type = MakeSimpleType(TypeName);
 		auto* MemberDeclaration = new(Allocator) FDeclaration(Allocator, SourceInfo);
-		MemberDeclaration->Identifier = TEXT("SimpleReturn");
+		MemberDeclaration->Identifier = FIdentifier::New(Allocator, TEXTVIEW("SimpleReturn"));
 		MemberDeclaration->Semantic = new(Allocator) FSemanticSpecifier(Allocator, Semantic, SourceInfo);
 		MemberDeclaratorList->Declarations.Add(MemberDeclaration);
 
@@ -662,7 +672,7 @@ struct FRemoveUnusedOutputs : FRemoveAlgorithm
 		// Create the LHS of the member assignment
 		FString MemberName = BodyContext.ReturnVariableName;
 		MemberName += TEXT(".");
-		MemberName += MemberDeclaration->Identifier;
+		MemberName += MemberDeclaration->Identifier->ToStringView();
 		auto* SimpleTypeMember = MakeIdentifierExpression(Allocator, *MemberName, SourceInfo);
 
 		// Create an assignment from the call the original function
@@ -734,12 +744,12 @@ struct FRemoveUnusedOutputs : FRemoveAlgorithm
 					}
 					else if (CheckSimpleVectorType(ParameterDeclarator->Type->Specifier->TypeName))
 					{
-						Errors.Add(FString::Printf(TEXT("RemoveUnusedOutputs: Function %s with out parameter %s doesn't have a return semantic"), *EntryPoint, ParameterDeclarator->Identifier));
+						Errors.Add(FString::Printf(TEXT("RemoveUnusedOutputs: Function %s with out parameter %s doesn't have a return semantic"), *EntryPoint, ParameterDeclarator->Identifier->GetData()));
 						return false;
 					}
 					else
 					{
-						Errors.Add(FString::Printf(TEXT("RemoveUnusedOutputs: Invalid return type %s for out parameter %s for function %s"), ParameterDeclarator->Type->Specifier->TypeName, ParameterDeclarator->Identifier, *EntryPoint));
+						Errors.Add(FString::Printf(TEXT("RemoveUnusedOutputs: Invalid return type %s for out parameter %s for function %s"), ParameterDeclarator->Type->Specifier->TypeName, ParameterDeclarator->Identifier->GetData(), *EntryPoint));
 						return false;
 					}
 				}
@@ -819,7 +829,7 @@ struct FRemoveUnusedOutputs : FRemoveAlgorithm
 							TCHAR* ElementSemanticPrefix = GetNonDigitSemanticPrefix(Allocator, MemberDeclaration->Semantic->Semantic, StartIndex);
 							if (!ElementSemanticPrefix)
 							{
-								Errors.Add(FString::Printf(TEXT("RemoveUnusedOutputs: Member (%s) %s : %s is expected to have an indexed semantic!"), MemberDeclarator->Type->Specifier->TypeName, MemberDeclaration->Identifier, MemberDeclaration->Semantic->Semantic));
+								Errors.Add(FString::Printf(TEXT("RemoveUnusedOutputs: Member (%s) %s : %s is expected to have an indexed semantic!"), MemberDeclarator->Type->Specifier->TypeName, MemberDeclaration->Identifier->GetData(), MemberDeclaration->Semantic->Semantic));
 
 								// Fatal: Array of non-indexed semantic (eg float4 Colors[4] : MYSEMANTIC; )
 								// Assume semantic is used and just fallback
@@ -838,7 +848,7 @@ struct FRemoveUnusedOutputs : FRemoveAlgorithm
 								{
 									auto* NewMemberDeclaration = new(Allocator) FDeclaration(Allocator, MemberDeclaration->SourceInfo);
 									NewMemberDeclaration->Semantic = new(Allocator) FSemanticSpecifier(Allocator, ElementSemantic, MemberDeclaration->SourceInfo);
-									NewMemberDeclaration->Identifier = Allocator->Strdup(FString::Printf(TEXT("%s_%d"), MemberDeclaration->Identifier, Index));
+									NewMemberDeclaration->Identifier = FIdentifier::New(Allocator, FString::Printf(TEXT("%s_%d"), MemberDeclaration->Identifier->GetData(), Index));
 
 									// Add member to struct
 									auto* NewDeclaratorList = new(Allocator) FDeclaratorList(Allocator, MemberDeclarator->SourceInfo);
@@ -846,8 +856,8 @@ struct FRemoveUnusedOutputs : FRemoveAlgorithm
 									NewDeclaratorList->Declarations.Add(NewMemberDeclaration);
 									DestStruct->Members.Add(NewDeclaratorList);
 
-									FString LHSElement = FString::Printf(TEXT("%s.%s"), DestPrefix, NewMemberDeclaration->Identifier);
-									FString RHSElement = FString::Printf(TEXT("%s.%s[%d]"), SourcePrefix, MemberDeclaration->Identifier, Index);
+									FString LHSElement = FString::Printf(TEXT("%s.%s"), DestPrefix, NewMemberDeclaration->Identifier->GetData());
+									FString RHSElement = FString::Printf(TEXT("%s.%s[%d]"), SourcePrefix, MemberDeclaration->Identifier->GetData(), Index);
 
 									auto* LHS = MakeIdentifierExpression(Allocator, *LHSElement, SourceInfo);
 									auto* RHS = MakeIdentifierExpression(Allocator, *RHSElement, SourceInfo);
@@ -879,7 +889,7 @@ struct FRemoveUnusedOutputs : FRemoveAlgorithm
 					{
 						if (!MemberDeclarator->Type || !MemberDeclarator->Type->Specifier || !MemberDeclarator->Type->Specifier->TypeName)
 						{
-							Errors.Add(FString::Printf(TEXT("RemoveUnusedOutputs: Internal error tracking down nested type %s"), MemberDeclaration->Identifier));
+							Errors.Add(FString::Printf(TEXT("RemoveUnusedOutputs: Internal error tracking down nested type %s"), MemberDeclaration->Identifier->GetData()));
 							return false;
 						}
 
@@ -887,14 +897,14 @@ struct FRemoveUnusedOutputs : FRemoveAlgorithm
 						FStructSpecifier* NestedStructSpecifier = FindStructSpecifier(MiniSymbolTable, MemberDeclarator->Type->Specifier->TypeName);
 						if (!NestedStructSpecifier)
 						{
-							Errors.Add(FString::Printf(TEXT("RemoveUnusedOutputs: Member (%s) %s is expected to have a semantic!"), MemberDeclarator->Type->Specifier->TypeName, MemberDeclaration->Identifier));
+							Errors.Add(FString::Printf(TEXT("RemoveUnusedOutputs: Member (%s) %s is expected to have a semantic!"), MemberDeclarator->Type->Specifier->TypeName, MemberDeclaration->Identifier->GetData()));
 							return false;
 						}
 
 						// Add all the elements of this new struct into the return type
 						FString NewSourcePrefix = SourcePrefix;
 						NewSourcePrefix += TEXT(".");
-						NewSourcePrefix += MemberDeclaration->Identifier;
+						NewSourcePrefix += MemberDeclaration->Identifier->ToStringView();
 						AddUsedOutputMembers(DestStruct, DestPrefix, NestedStructSpecifier, *NewSourcePrefix, MiniSymbolTable, BodyContext);
 					}
 				}
@@ -1071,7 +1081,7 @@ struct FRemoveUnusedInputs : FRemoveAlgorithm
 			// If the input struct is not empty, add this as an argument to the new entry function
 			FParameterDeclarator* Declarator = new(Allocator) FParameterDeclarator(Allocator, SourceInfo);
 			Declarator->Type = MakeSimpleType(BodyContext.InputTypeName);
-			Declarator->Identifier = BodyContext.InputVariableName;
+			Declarator->Identifier = FIdentifier::New(Allocator, BodyContext.InputVariableName);
 			BodyContext.NewFunctionParameters.Add(Declarator);
 		}
 
@@ -1080,7 +1090,7 @@ struct FRemoveUnusedInputs : FRemoveAlgorithm
 
 		WriteGeneratedInCode(NewEntryFunction, BodyContext.NewStructs, GeneratedCode);
 
-		EntryPoint = NewEntryFunction->Prototype->Identifier;
+		EntryPoint = NewEntryFunction->Prototype->Identifier->ToString();
 		bSuccess = true;
 	}
 
@@ -1110,12 +1120,12 @@ struct FRemoveUnusedInputs : FRemoveAlgorithm
 					}
 					else if (CheckSimpleVectorType(ParameterDeclarator->Type->Specifier->TypeName))
 					{
-						Errors.Add(FString::Printf(TEXT("RemoveUnusedInputs: Function %s with in parameter %s doesn't have a return semantic"), *EntryPoint, ParameterDeclarator->Identifier));
+						Errors.Add(FString::Printf(TEXT("RemoveUnusedInputs: Function %s with in parameter %s doesn't have a return semantic"), *EntryPoint, ParameterDeclarator->Identifier->GetData()));
 						return false;
 					}
 					else
 					{
-						Errors.Add(FString::Printf(TEXT("RemoveUnusedInputs: Invalid return type %s for in parameter %s for function %s"), ParameterDeclarator->Type->Specifier->TypeName, ParameterDeclarator->Identifier, *EntryPoint));
+						Errors.Add(FString::Printf(TEXT("RemoveUnusedInputs: Invalid return type %s for in parameter %s for function %s"), ParameterDeclarator->Type->Specifier->TypeName, ParameterDeclarator->Identifier->GetData(), *EntryPoint));
 						return false;
 					}
 				}
@@ -1135,18 +1145,17 @@ struct FRemoveUnusedInputs : FRemoveAlgorithm
 
 		auto* Zero = new(Allocator) FExpression(Allocator, EOperators::Literal, SourceInfo);
 		Zero->LiteralType = CrossCompiler::ELiteralType::Integer;
-		Zero->Identifier = Allocator->Strdup(TEXT("0"));
+		Zero->Identifier = FIdentifier::New(Allocator,TEXTVIEW("0"));
 		auto* Initializer = new(Allocator) FUnaryExpression(Allocator, EOperators::TypeCast, Zero, SourceInfo);
 		Initializer->TypeSpecifier = MakeSimpleType(OriginalStructSpecifier->Name)->Specifier;
 
 		// Add a local variable to receive the output from the function.
 		// The name is simply the parameter name in the real entry point, since it will be inlined directly in the optimized version.
 
-		const TCHAR* LocalStructVarName = ParameterDeclarator->Identifier;
-		auto* LocalStructVariable = CreateLocalVariable(OriginalStructSpecifier->Name, LocalStructVarName, Initializer);
+		auto* LocalStructVariable = CreateLocalVariable(OriginalStructSpecifier->Name, ParameterDeclarator->Identifier, Initializer);
 		BodyContext.PreInstructions.Add(LocalStructVariable);
 
-		return AddUsedInputMembers(BodyContext.NewInputStruct, BodyContext.InputVariableName, OriginalStructSpecifier, LocalStructVarName, MiniSymbolTable, BodyContext);
+		return AddUsedInputMembers(BodyContext.NewInputStruct, BodyContext.InputVariableName, OriginalStructSpecifier, ParameterDeclarator->Identifier->GetData(), MiniSymbolTable, BodyContext);
 	}
 
 	bool IsSemanticUsed(const TCHAR* SemanticName)
@@ -1196,7 +1205,7 @@ struct FRemoveUnusedInputs : FRemoveAlgorithm
 			// Make this parameter the initializer of the new local variable
 			FString ParameterName = BodyContext.InputVariableName;
 			ParameterName += TEXT(".");
-			ParameterName += ParameterDeclarator->Identifier;
+			ParameterName += ParameterDeclarator->Identifier->ToStringView();
 			Initializer = MakeIdentifierExpression(Allocator, *ParameterName, SourceInfo);
 		}
 		else
@@ -1204,7 +1213,7 @@ struct FRemoveUnusedInputs : FRemoveAlgorithm
 			// Make a local to generate the in parameter: Type Local = (Type)0;
 			auto* Zero = new(Allocator) FExpression(Allocator, EOperators::Literal, SourceInfo);
 			Zero->LiteralType = CrossCompiler::ELiteralType::Integer;
-			Zero->Identifier = Allocator->Strdup(TEXT("0"));
+			Zero->Identifier = FIdentifier::New(Allocator, TEXTVIEW("0"));
 			Initializer = new(Allocator) FUnaryExpression(Allocator, EOperators::TypeCast, Zero, SourceInfo);
 			Initializer->TypeSpecifier = ParameterDeclarator->Type->Specifier;
 
@@ -1219,7 +1228,7 @@ struct FRemoveUnusedInputs : FRemoveAlgorithm
 	{
 		CrossCompiler::AST::FASTWriter Writer(OutGeneratedCode);
 		GeneratedCode = TEXT("#line 1 \"RemoveUnusedInputs.usf\"\n// Generated Entry Point: ");
-		GeneratedCode += NewEntryFunction->Prototype->Identifier;
+		GeneratedCode += NewEntryFunction->Prototype->Identifier->ToStringView();
 		GeneratedCode += TEXT("\n");
 		if (UsedInputs.Num() > 0)
 		{
@@ -1278,7 +1287,7 @@ struct FRemoveUnusedInputs : FRemoveAlgorithm
 							TCHAR* ElementSemanticPrefix = GetNonDigitSemanticPrefix(Allocator, MemberDeclaration->Semantic->Semantic, StartIndex);
 							if (!ElementSemanticPrefix)
 							{
-								Errors.Add(FString::Printf(TEXT("RemoveUnusedInputs: Member (%s) %s : %s is expected to have an indexed semantic!"), MemberDeclarator->Type->Specifier->TypeName, MemberDeclaration->Identifier, MemberDeclaration->Semantic->Semantic));
+								Errors.Add(FString::Printf(TEXT("RemoveUnusedInputs: Member (%s) %s : %s is expected to have an indexed semantic!"), MemberDeclarator->Type->Specifier->TypeName, MemberDeclaration->Identifier->GetData(), MemberDeclaration->Semantic->Semantic));
 
 								// Fatal: Array of non-indexed semantic (eg float4 Colors[4] : MYSEMANTIC; )
 								// Assume semantic is used and just fallback
@@ -1298,7 +1307,7 @@ struct FRemoveUnusedInputs : FRemoveAlgorithm
 								{
 									auto* NewMemberDeclaration = new(Allocator) FDeclaration(Allocator, MemberDeclaration->SourceInfo);
 									NewMemberDeclaration->Semantic = new(Allocator) FSemanticSpecifier(Allocator, ElementSemantic, MemberDeclaration->SourceInfo);
-									NewMemberDeclaration->Identifier = Allocator->Strdup(FString::Printf(TEXT("%s_%d"), MemberDeclaration->Identifier, Index));
+									NewMemberDeclaration->Identifier = FIdentifier::New(Allocator, FString::Printf(TEXT("%s_%d"), MemberDeclaration->Identifier->GetData(), Index));
 
 									// Add member to struct
 									auto* NewDeclaratorList = new(Allocator) FDeclaratorList(Allocator, MemberDeclarator->SourceInfo);
@@ -1306,8 +1315,8 @@ struct FRemoveUnusedInputs : FRemoveAlgorithm
 									NewDeclaratorList->Declarations.Add(NewMemberDeclaration);
 									DestStruct->Members.Add(NewDeclaratorList);
 
-									FString LHSElement = FString::Printf(TEXT("%s.%s[%d]"), SourcePrefix, MemberDeclaration->Identifier, Index);
-									FString RHSElement = FString::Printf(TEXT("%s.%s"), DestPrefix, NewMemberDeclaration->Identifier);
+									FString LHSElement = FString::Printf(TEXT("%s.%s[%d]"), SourcePrefix, MemberDeclaration->Identifier->GetData(), Index);
+									FString RHSElement = FString::Printf(TEXT("%s.%s"), DestPrefix, NewMemberDeclaration->Identifier->GetData());
 
 									auto* LHS = MakeIdentifierExpression(Allocator, *LHSElement, SourceInfo);
 									auto* RHS = MakeIdentifierExpression(Allocator, *RHSElement, SourceInfo);
@@ -1350,7 +1359,7 @@ struct FRemoveUnusedInputs : FRemoveAlgorithm
 					{
 						if (!MemberDeclarator->Type || !MemberDeclarator->Type->Specifier || !MemberDeclarator->Type->Specifier->TypeName)
 						{
-							Errors.Add(FString::Printf(TEXT("RemoveUnusedInputs: Internal error tracking down nested type %s"), MemberDeclaration->Identifier));
+							Errors.Add(FString::Printf(TEXT("RemoveUnusedInputs: Internal error tracking down nested type %s"), MemberDeclaration->Identifier->GetData()));
 							return false;
 						}
 
@@ -1358,14 +1367,14 @@ struct FRemoveUnusedInputs : FRemoveAlgorithm
 						FStructSpecifier* NestedStructSpecifier = FindStructSpecifier(MiniSymbolTable, MemberDeclarator->Type->Specifier->TypeName);
 						if (!NestedStructSpecifier)
 						{
-							Errors.Add(FString::Printf(TEXT("RemoveUnusedInputs: Member (%s) %s is expected to have a semantic!"), MemberDeclarator->Type->Specifier->TypeName, MemberDeclaration->Identifier));
+							Errors.Add(FString::Printf(TEXT("RemoveUnusedInputs: Member (%s) %s is expected to have a semantic!"), MemberDeclarator->Type->Specifier->TypeName, MemberDeclaration->Identifier->GetData()));
 							return false;
 						}
 
 						// Add all the elements of this new struct into the return type
 						FString NewSourcePrefix = SourcePrefix;
 						NewSourcePrefix += TEXT(".");
-						NewSourcePrefix += MemberDeclaration->Identifier;
+						NewSourcePrefix += MemberDeclaration->Identifier->ToStringView();
 						AddUsedInputMembers(DestStruct, DestPrefix, NestedStructSpecifier, *NewSourcePrefix, MiniSymbolTable, BodyContext);
 					}
 				}
@@ -1473,7 +1482,7 @@ struct FFindEntryPointParameters
 			}
 			else if (FFunctionDefinition* FunctionDefinition = Node->AsFunctionDefinition())
 			{
-				if (FCString::Strcmp(*EntryPoint, FunctionDefinition->Prototype->Identifier) == 0)
+				if (FunctionDefinition->Prototype->Identifier->Equals(EntryPoint))
 				{
 					EntryFunction = FunctionDefinition;
 				}
@@ -1529,7 +1538,7 @@ struct FFindEntryPointParameters
 				{
 					if (FSemanticSpecifier* Semantic = MemberDecl->Semantic)
 					{
-						if (!ProcessParameterSemantic(*Semantic, MemberDecl->Identifier, bIsOutput))
+						if (!ProcessParameterSemantic(*Semantic, MemberDecl->Identifier->GetData(), bIsOutput))
 						{
 							return false;
 						}
@@ -1543,12 +1552,12 @@ struct FFindEntryPointParameters
 					}
 					else if (CheckSimpleVectorType(DeclList->Type->Specifier->TypeName))
 					{
-						Errors.Add(FString::Printf(TEXT("FindEntryPointParameters: Function %s with parameter %s doesn't have a return semantic"), *EntryPoint, MemberDecl->Identifier));
+						Errors.Add(FString::Printf(TEXT("FindEntryPointParameters: Function %s with parameter %s doesn't have a return semantic"), *EntryPoint, MemberDecl->Identifier->GetData()));
 						return false;
 					}
 					else
 					{
-						Errors.Add(FString::Printf(TEXT("FindEntryPointParameters: Invalid return type %s for parameter %s in function %s"), DeclList->Type->Specifier->TypeName, MemberDecl->Identifier, *EntryPoint));
+						Errors.Add(FString::Printf(TEXT("FindEntryPointParameters: Invalid return type %s for parameter %s in function %s"), DeclList->Type->Specifier->TypeName, MemberDecl->Identifier->GetData(), *EntryPoint));
 						return false;
 					}
 				}
@@ -1567,7 +1576,7 @@ struct FFindEntryPointParameters
 			const bool bIsOutput = ParameterDeclarator->Type->Qualifier.bOut;
 			if (FSemanticSpecifier* Semantic = ParameterDeclarator->Semantic)
 			{
-				if (!ProcessParameterSemantic(*Semantic, ParameterDeclarator->Identifier, bIsOutput))
+				if (!ProcessParameterSemantic(*Semantic, ParameterDeclarator->Identifier->GetData(), bIsOutput))
 				{
 					return false;
 				}
@@ -1581,12 +1590,12 @@ struct FFindEntryPointParameters
 			}
 			else if (CheckSimpleVectorType(ParameterDeclarator->Type->Specifier->TypeName))
 			{
-				Errors.Add(FString::Printf(TEXT("FindEntryPointParameters: Function %s with parameter %s doesn't have a return semantic"), *EntryPoint, ParameterDeclarator->Identifier));
+				Errors.Add(FString::Printf(TEXT("FindEntryPointParameters: Function %s with parameter %s doesn't have a return semantic"), *EntryPoint, ParameterDeclarator->Identifier->GetData()));
 				return false;
 			}
 			else
 			{
-				Errors.Add(FString::Printf(TEXT("FindEntryPointParameters: Invalid return type %s for parameter %s in function %s"), ParameterDeclarator->Type->Specifier->TypeName, ParameterDeclarator->Identifier, *EntryPoint));
+				Errors.Add(FString::Printf(TEXT("FindEntryPointParameters: Invalid return type %s for parameter %s in function %s"), ParameterDeclarator->Type->Specifier->TypeName, ParameterDeclarator->Identifier->GetData(), *EntryPoint));
 				return false;
 			}
 		}
@@ -1667,34 +1676,42 @@ struct FConvertFP32ToFP16 {
 	bool bSuccess;
 };
 
-static void ConvertFromFP32ToFP16(const TCHAR*& TypeName, CrossCompiler::FLinearAllocator* Allocator)
+static FStringView FindFP16Type(const FStringView TypeName)
 {
-	static FString FloatTypes[9] = { "float", "float2", "float3", "float4", "float2x2", "float3x3", "float4x4", "float3x4", "float4x3" };
-	static FString HalfTypes[9] = { "half", "half2", "half3", "half4", "half2x2", "half3x3", "half4x4", "half3x4", "half4x3" };
-	//static FString HalfTypes[9] = { "min16float", "min16float2", "min16float3", "min16float4", "min16float2x2", "min16float3x3", "min16float4x4", "min16float3x4", "min16float4x3" };
-	FString NewType;
-	for (int32 i = 0; i < 9; ++i) 
+	static const FStringView FloatTypes[9] = { TEXTVIEW("float"), TEXTVIEW("float2"), TEXTVIEW("float3"), TEXTVIEW("float4"), TEXTVIEW("float2x2"), TEXTVIEW("float3x3"), TEXTVIEW("float4x4"), TEXTVIEW("float3x4"), TEXTVIEW("float4x3") };
+	static const FStringView HalfTypes[9] = { TEXTVIEW("half"), TEXTVIEW("half2"), TEXTVIEW("half3"), TEXTVIEW("half4"), TEXTVIEW("half2x2"), TEXTVIEW("half3x3"), TEXTVIEW("half4x4"), TEXTVIEW("half3x4"), TEXTVIEW("half4x3") };
+	//static const FStringView HalfTypes[9] = { TEXTVIEW("min16float"), TEXTVIEW("min16float2"), TEXTVIEW("min16float3"), TEXTVIEW("min16float4"), TEXTVIEW("min16float2x2"), TEXTVIEW("min16float3x3"), TEXTVIEW("min16float4x4"), TEXTVIEW("min16float3x4"), TEXTVIEW("min16float4x3") };
+	for (int32 i = 0; i < 9; ++i)
 	{
-		if (FString(TypeName).Equals(FloatTypes[i]))
+		if (TypeName.Equals(FloatTypes[i]))
 		{
-			NewType = HalfTypes[i];
-			break;
+			return HalfTypes[i];
 		}
 	}
+
+	return {};
+}
+
+static void ConvertFromFP32ToFP16(const CrossCompiler::AST::FIdentifier*& IdentifierName, CrossCompiler::FLinearAllocator* Allocator)
+{
+	FStringView NewType = FindFP16Type(IdentifierName->ToStringView());
 	if (NewType.IsEmpty()) return;
 	check(Allocator);
-	*const_cast<TCHAR**>(&TypeName) = Allocator->Strdup(*NewType);
+	IdentifierName = CrossCompiler::AST::FIdentifier::New(Allocator, NewType);
 }
 
 static void ConvertFromFP32ToFP16(CrossCompiler::AST::FTypeSpecifier* Type, CrossCompiler::FLinearAllocator* Allocator)
 {
-    ConvertFromFP32ToFP16(Type->TypeName, Allocator);
+	FStringView NewType = FindFP16Type(FStringView(Type->TypeName));
+	if (NewType.IsEmpty()) return;
+	check(Allocator);
+	Type->TypeName = Allocator->Strdup(NewType.GetData(), NewType.Len());
 }
 
 static void ConvertFromFP32ToFP16Base(CrossCompiler::AST::FNode* Node, CrossCompiler::FLinearAllocator* Allocator);
 static void ConvertFromFP32ToFP16(CrossCompiler::AST::FFunctionDefinition* Node, CrossCompiler::FLinearAllocator* Allocator)
 {
-	if (FString(Node->Prototype->Identifier).Equals("CalcSceneDepth"))
+	if (Node->Prototype->Identifier->Equals(TEXTVIEW("CalcSceneDepth")))
 	{
 		return;
 	}
@@ -1714,7 +1731,7 @@ static void ConvertFromFP32ToFP16(CrossCompiler::AST::FParameterDeclarator* Node
 {
 	if (Node->bIsArray)
 	{
-		if (!FString("MRT").Equals(Node->Identifier))
+		if (!Node->Identifier->Equals(TEXTVIEW("MRT")))
 		{
 			return;
 		}
@@ -1728,7 +1745,7 @@ static void ConvertFromFP32ToFP16(CrossCompiler::AST::FDeclaratorList* Node, Cro
 	{
 		if (Elem->AsDeclaration() && Elem->AsDeclaration()->bIsArray)
 		{
-			if (!FString("MRT").Equals(Elem->AsDeclaration()->Identifier))
+			if (!Elem->AsDeclaration()->Identifier->Equals(TEXTVIEW("MRT")))
 			{
 				return;
 			}
@@ -2023,16 +2040,16 @@ namespace UE::ShaderCompilerCommon
 				}
 			}
 
-			void ValidateDeclarationIdentifier(const TCHAR* InIdentifier, const CrossCompiler::FSourceInfo* SourceInfo = nullptr)
+			void ValidateDeclarationIdentifier(const CrossCompiler::AST::FIdentifier* InIdentifier, const CrossCompiler::FSourceInfo* SourceInfo = nullptr)
 			{
 				// Validate identifier against fixed set of names that are known to cause problems with FXC
-				if (FCString::Strcmp(InIdentifier, TEXT("sample")) == 0)
+				if (InIdentifier->Equals(TEXTVIEW("sample")))
 				{
-					Error(FString::Printf(TEXT("Identifier \"%s\" must not be used as parameter or declaration identifier. FXC (D3D11) will misinterpret it as interpolation qualifier type. Consider renaming it to \"Sample\"."), InIdentifier), SourceInfo);
+					Error(FString::Printf(TEXT("Identifier \"%s\" must not be used as parameter or declaration identifier. FXC (D3D11) will misinterpret it as interpolation qualifier type. Consider renaming it to \"Sample\"."), InIdentifier->GetData()), SourceInfo);
 				}
-				if (FCString::Strcmp(InIdentifier, TEXT("Buffer")) == 0)
+				if (InIdentifier->Equals(TEXTVIEW("Buffer")))
 				{
-					Error(FString::Printf(TEXT("Identifier \"%s\" must not be used as parameter or declaration identifier. FXC (D3D11) will misinterpret it as resource type. Consider adding the \"In\" or \"Out\" prefix."), InIdentifier), SourceInfo);
+					Error(FString::Printf(TEXT("Identifier \"%s\" must not be used as parameter or declaration identifier. FXC (D3D11) will misinterpret it as resource type. Consider adding the \"In\" or \"Out\" prefix."), InIdentifier->GetData()), SourceInfo);
 				}
 			}
 
@@ -2095,6 +2112,15 @@ namespace UE::ShaderCompilerCommon
 
 #if WITH_LOW_LEVEL_TESTS
 
+static void GenerateShaderCode(const CrossCompiler::TLinearArray<CrossCompiler::AST::FNode*>& ASTNodes, FString& OutResult)
+{
+	CrossCompiler::AST::FASTWriter Writer(OutResult);
+	for (const CrossCompiler::AST::FNode* Node : ASTNodes)
+	{
+		Node->Write(Writer);
+	}
+}
+
 static TArray<FString> GetShaderInputSemantics(const FString& ShaderSource, const FString& EntryPoint)
 {
 	TArray<FString> Result;
@@ -2120,7 +2146,7 @@ static TArray<FString> GetShaderInputSemantics(const FString& ShaderSource, cons
 				{
 					FFunction* Prototype = FunctionDefinition->Prototype;
 
-					if (!Prototype || FCString::Strcmp(Prototype->Identifier, *EntryPoint))
+					if (!Prototype || !Prototype->Identifier->Equals(EntryPoint))
 					{
 						continue;
 					}
@@ -2287,4 +2313,151 @@ float4 main(FPSInputs Inputs, float4 C : TEXCOORD2) : SV_Target
 	TEST_TRUE(TEXT("Rewritten shader uses TEXCOORD2"), RewrittenInputs.Contains(TEXT("TEXCOORD2")));
 }
 
+TEST_CASE_NAMED(FShaderRenameLocalsTest, "Shaders::RenameLocals", "[EditorContext][EngineFilter]")
+{
+	FString ShaderCodeString(
+		TEXT(R"(
+float fun(float Param_A, float Param_B)
+{
+	float C = 123.0;
+	return Param_B * Param_A + C;
+}
+float4 main(float4 Param_C : TEXCOORD0) : SV_Target
+{
+	// test comment
+	float A = 0.1;
+	float B = 0.2 + A;
+	for (int I=0; I < 4; ++I)
+	{
+		B *= 2.0;
+	}
+	for (int J=0, K=1; J < 4; ++J)
+	{
+		B *= 2.0;
+		B += (float)K;
+	}
+
+	int X = int(Param_C.x);
+	switch (X)
+	{
+	case 0:
+		A *= 2.0;
+		break;
+	default:
+		break;
+	}
+
+	return float4(A, B, A+B, fun(A,B)) + Param_C;
+}
+)"));
+
+	FString EntryPoint(TEXT("main"));
+
+	FString Result;
+
+	CrossCompiler::FCompilerMessages CompilerMessages;
+	TArray<FScopedDeclarations> ScopedDeclarations;
+	CrossCompiler::Parser::Parse(ShaderCodeString, TEXT("shader.hlsl"), CompilerMessages, ScopedDeclarations,
+		[&Result, &EntryPoint](CrossCompiler::FLinearAllocator* Allocator, CrossCompiler::TLinearArray<CrossCompiler::AST::FNode*>& ASTNodes)
+		{
+			using namespace CrossCompiler::AST;
+
+			auto ProcessIdentifier = [Allocator](FIdentifier* Identifier)
+			{
+				TStringBuilder<256> NewName;
+				NewName << TEXT("Renamed_");
+				NewName << Identifier->ToStringView();
+
+				Identifier->Rename(Allocator, NewName);
+			};
+
+			auto ProcessDeclaration = [&ProcessIdentifier](FDeclaration* DeclarationNode)
+				{
+					if (!DeclarationNode || !DeclarationNode->Identifier)
+					{
+						return;
+					}
+
+					ProcessIdentifier(DeclarationNode->Identifier);
+				};
+
+			auto ProcessDeclaratorList = [&ProcessDeclaration](FDeclaratorList* DeclaratorListNode)
+			{
+				for (FNode* InnerDeclarationNode : DeclaratorListNode->Declarations)
+				{
+					if (FDeclaration* InnerDeclaration = InnerDeclarationNode->AsDeclaration())
+					{
+						ProcessDeclaration(InnerDeclaration);
+					}
+				}
+			};
+
+			for (FNode* Node : ASTNodes)
+			{
+				FFunctionDefinition* FunctionDefinition = Node->AsFunctionDefinition();
+				if (!FunctionDefinition || !FunctionDefinition->Body || !FunctionDefinition->Prototype->Identifier->Equals(EntryPoint))
+				{
+					continue;
+				}
+
+				for (FNode* ParameterNode : FunctionDefinition->Prototype->Parameters)
+				{
+					FParameterDeclarator* ParameterDeclarator = ParameterNode->AsParameterDeclarator();
+					if (!ParameterDeclarator)
+					{
+						continue;
+					}
+
+					if (ParameterDeclarator->Identifier)
+					{
+						ProcessIdentifier(ParameterDeclarator->Identifier);
+					}
+				}
+
+				TArray<FCompoundStatement*> Stack;
+				Stack.Push(FunctionDefinition->Body);
+
+				while (!Stack.IsEmpty())
+				{
+					FCompoundStatement* ParentCompoundStatementNode = Stack.Last();
+					Stack.Pop();
+
+					for (FNode* StatementNode : ParentCompoundStatementNode->Statements)
+					{
+						if (FCompoundStatement* CompoundStatement = StatementNode->AsCompoundStatement())
+						{
+							Stack.Push(CompoundStatement);
+						}
+						else if (FIterationStatement* IterationStatement = StatementNode->AsIterationStatement())
+						{
+							FNode* InitStatementNode = IterationStatement->InitStatement;
+							if (FDeclaratorList* DeclaratorListNode = InitStatementNode->AsDeclaratorList())
+							{
+								ProcessDeclaratorList(DeclaratorListNode);
+							}
+						}
+						else if (FDeclaration* Declaration = StatementNode->AsDeclaration())
+						{
+							ProcessDeclaration(Declaration);
+						}
+						else if (FDeclaratorList* DeclaratorListNode = StatementNode->AsDeclaratorList())
+						{
+							ProcessDeclaratorList(DeclaratorListNode);
+						}
+					}
+				}
+			}
+
+			GenerateShaderCode(ASTNodes, Result);
+		});
+
+	// Just a basic check at this point.
+	// TODO: we could parse the rewritten shader and validate that old identifiers are not present and new ones are.
+
+	TEST_TRUE(TEXT("Function signature uses renamed parameters"), Result.Contains(TEXT("float4 main(float4 Renamed_Param_C : TEXCOORD0 ) : SV_Target")));
+	TEST_TRUE(TEXT("Function uses renamed for loop variable I"), Result.Contains(TEXT("for (int Renamed_I = 0; Renamed_I < 4; ++Renamed_I)")));
+	TEST_TRUE(TEXT("Function uses renamed for loop variables J and K"), Result.Contains(TEXT("for (int Renamed_J = 0, Renamed_K = 1; Renamed_J < 4; ++Renamed_J)")));
+	TEST_TRUE(TEXT("Function uses renamed locals"), Result.Contains(TEXT("float4(Renamed_A, Renamed_B, (Renamed_A + Renamed_B), fun(Renamed_A, Renamed_B)) + Renamed_Param_C")));
+}
+
 #endif // WITH_LOW_LEVEL_TESTS
diff --git a/Engine/Source/Developer/ShaderCompilerCommon/Private/HlslUtils.h b/Engine/Source/Developer/ShaderCompilerCommon/Private/HlslUtils.h
--- a/Engine/Source/Developer/ShaderCompilerCommon/Private/HlslUtils.h
+++ b/Engine/Source/Developer/ShaderCompilerCommon/Private/HlslUtils.h
@@ -90,23 +90,29 @@ namespace CrossCompiler
 			return (void*)Address;
 		}
 
-		/*inline*/ TCHAR* Strdup(const TCHAR* String)
+		TCHAR* Strdup(const TCHAR* String, int32 Length = -1)
 		{
 			if (!String)
 			{
 				return nullptr;
 			}
 
-			const auto Length = FCString::Strlen(String);
+			if (Length < 0)
+			{
+				Length = FCString::Strlen(String);
+			}
+
 			const auto Size = (Length + 1) * sizeof(TCHAR);
 			auto* Data = (TCHAR*)Alloc(Size, sizeof(TCHAR));
-			FCString::Strncpy(Data, String, Length+1);
+
+			FCString::Strncpy(Data, String, Length + 1);
+
 			return Data;
 		}
 
 		inline TCHAR* Strdup(const FString& String)
 		{
-			return Strdup(*String);
+			return Strdup(*String, String.Len());
 		}
 
 		TArray<Memory::FPage*, TInlineAllocator<8> > Pages;
