Fix symbol scope tracking for function parameters

Function parameters were previously put into the outer scope, while a new scope was created for the function body. Now the arguments and the body share the same scope, while function name itself is placed into the parent scope.

diff --git a/Engine/Source/Developer/ShaderCompilerCommon/Private/HlslExpressionParser.inl b/Engine/Source/Developer/ShaderCompilerCommon/Private/HlslExpressionParser.inl
--- a/Engine/Source/Developer/ShaderCompilerCommon/Private/HlslExpressionParser.inl
+++ b/Engine/Source/Developer/ShaderCompilerCommon/Private/HlslExpressionParser.inl
@@ -27,6 +27,23 @@ namespace CrossCompiler
 
 	struct FSymbolScope
 	{
+		enum class EType
+		{
+			Unknown,
+			Global,
+			Namespace,
+			Statement,
+			Do,
+			While,
+			For,
+			If,
+			ScopedDeclaration,
+			Function,
+		};
+
+		// Type tag is for debugging
+		const EType ScopeType;
+
 		FLinearAllocator* Allocator;
 		FSymbolScope* Parent;
 		const TCHAR* Name;
@@ -48,8 +65,9 @@ namespace CrossCompiler
 
 		TMap<FStringView, CrossCompiler::AST::FIdentifier*, FDefaultSetAllocator, FIdentifierKeyFuncs> Identifiers;
 
-		FSymbolScope(FLinearAllocator* InAllocator, FSymbolScope* InParent) 
-			: Allocator(InAllocator)
+		FSymbolScope(FLinearAllocator* InAllocator, FSymbolScope* InParent, EType InScopeType)
+			: ScopeType(InScopeType)
+			, Allocator(InAllocator)
 			, Parent(InParent)
 			, Name(nullptr)
 			, Children(InAllocator)
@@ -172,14 +190,16 @@ namespace CrossCompiler
 
 	struct FCreateSymbolScope
 	{
+		UE_NONCOPYABLE(FCreateSymbolScope)
+
 		FSymbolScope* Original;
 		FSymbolScope** Current;
 
-		FCreateSymbolScope(FLinearAllocator* InAllocator, FSymbolScope** InCurrent) :
+		FCreateSymbolScope(FLinearAllocator* InAllocator, FSymbolScope** InCurrent, FSymbolScope::EType ScopeType) :
 			Current(InCurrent)
 		{
 			Original = *InCurrent;
-			auto& NewScope = Original->Children.Emplace_GetRef(InAllocator, Original);
+			auto& NewScope = Original->Children.Emplace_GetRef(InAllocator, Original, ScopeType);
 			*Current = &NewScope;
 		}
 
diff --git a/Engine/Source/Developer/ShaderCompilerCommon/Private/HlslParser.cpp b/Engine/Source/Developer/ShaderCompilerCommon/Private/HlslParser.cpp
--- a/Engine/Source/Developer/ShaderCompilerCommon/Private/HlslParser.cpp
+++ b/Engine/Source/Developer/ShaderCompilerCommon/Private/HlslParser.cpp
@@ -7,6 +7,7 @@
 #include "HlslParser.h"
 #include "HlslExpressionParser.inl"
 #include "ShaderCompilerCommon.h"
+#include "Misc/Optional.h"
 
 namespace CrossCompiler
 {
@@ -1182,7 +1183,7 @@ namespace CrossCompiler
 		return ParseFullType(Parser.Scanner, TypeFlags, TemplateTypeFlags, Parser.CurrentScope, Allocator, OutFullySpecifiedType);
 	}
 
-	EParseResult ParseFunctionDeclarator(FHlslParser& Parser, FLinearAllocator* Allocator, AST::FFunction** OutFunction)
+	EParseResult ParseFunctionDeclarator(FHlslParser& Parser, FLinearAllocator* Allocator, AST::FFunction** OutFunction, TOptional<FCreateSymbolScope>* OutScope)
 	{
 		uint32 OriginalToken = Parser.Scanner.GetCurrentTokenIndex();
 		AST::FFullySpecifiedType* ReturnType = nullptr;
@@ -1276,6 +1277,11 @@ namespace CrossCompiler
 		}
 		else
 		{
+			if (OutScope)
+			{
+				OutScope->Emplace(Allocator, &Parser.CurrentScope, FSymbolScope::EType::Function);
+			}
+
 			Result = ParseFunctionParameterDeclaration(Parser, Allocator, Function);
 			if (Result == EParseResult::Error)
 			{
@@ -1321,7 +1327,7 @@ Done:
 			Token1 && Token1->Token == EHlslToken::Identifier && 
 			Token2 && Token2->Token == EHlslToken::LeftBrace)
 		{
-			FCreateSymbolScope NamespaceSymbolScope(Allocator, &Parser.CurrentScope);
+			FCreateSymbolScope NamespaceSymbolScope(Allocator, &Parser.CurrentScope, FSymbolScope::EType::Namespace);
 			Parser.CurrentScope->Name = *(Token1->String);
 
 			check(Parser.Scanner.MatchToken(EHlslToken::Namespace));
@@ -1345,7 +1351,7 @@ Done:
 	
 	EParseResult ParseStatementBlock(FHlslParser& Parser, FLinearAllocator* Allocator, AST::FNode** OutStatement)
 	{
-		FCreateSymbolScope SymbolScope(Allocator, &Parser.CurrentScope);
+		FCreateSymbolScope SymbolScope(Allocator, &Parser.CurrentScope, FSymbolScope::EType::Statement);
 		auto* Block = new(Allocator) AST::FCompoundStatement(Allocator, Parser.Scanner.GetCurrentToken()->SourceInfo);
 		// To help debug
 		AST::FNode* PrevStatement = nullptr;
@@ -1426,8 +1432,13 @@ Done:
 			return ParseResultError();
 		}
 
+		// Symbol scope to be used for function parameters and body, but not the function name itself.
+		// We use TOptional here because FCreateSymbolScope destructor should run here (after finishing parsing the body), 
+		// but it will be constructed inside ParseFunctionDeclarator.
+		TOptional<FCreateSymbolScope> FunctionSymbolScope;
+
 		AST::FFunction* Function = nullptr;
-		EParseResult Result = ParseFunctionDeclarator(Parser, Allocator, &Function);
+		EParseResult Result = ParseFunctionDeclarator(Parser, Allocator, &Function, &FunctionSymbolScope);
 		if (Result == EParseResult::NotMatched)
 		{
 			if (bFoundInline)
@@ -1655,7 +1666,7 @@ Done:
 
 	EParseResult ParseDoStatement(FHlslParser& Parser, FLinearAllocator* Allocator, AST::FNode** OutStatement)
 	{
-		FCreateSymbolScope SymbolScope(Allocator, &Parser.CurrentScope);
+		FCreateSymbolScope SymbolScope(Allocator, &Parser.CurrentScope, FSymbolScope::EType::Do);
 		const auto* Token = Parser.Scanner.GetCurrentToken();
 		AST::FNode* Body = nullptr;
 		auto Result = ParseStatement(Parser, Allocator, &Body);
@@ -1704,7 +1715,7 @@ Done:
 
 	EParseResult ParseWhileStatement(FHlslParser& Parser, FLinearAllocator* Allocator, AST::FNode** OutStatement)
 	{
-		FCreateSymbolScope SymbolScope(Allocator, &Parser.CurrentScope);
+		FCreateSymbolScope SymbolScope(Allocator, &Parser.CurrentScope, FSymbolScope::EType::While);
 		const auto* Token = Parser.Scanner.GetCurrentToken();
 		if (!Parser.Scanner.MatchToken(EHlslToken::LeftParenthesis))
 		{
@@ -1741,7 +1752,7 @@ Done:
 
 	EParseResult ParseForStatement(FHlslParser& Parser, FLinearAllocator* Allocator, AST::FNode** OutStatement)
 	{
-		FCreateSymbolScope SymbolScope(Allocator, &Parser.CurrentScope);
+		FCreateSymbolScope SymbolScope(Allocator, &Parser.CurrentScope, FSymbolScope::EType::For);
 		const auto* Token = Parser.Scanner.GetCurrentToken();
 		if (!Parser.Scanner.MatchToken(EHlslToken::LeftParenthesis))
 		{
@@ -1844,7 +1855,7 @@ Done:
 
 	EParseResult ParseIfStatement(FHlslParser& Parser, FLinearAllocator* Allocator, AST::FNode** OutStatement)
 	{
-		FCreateSymbolScope SymbolScope(Allocator, &Parser.CurrentScope);
+		FCreateSymbolScope SymbolScope(Allocator, &Parser.CurrentScope, FSymbolScope::EType::If);
 
 		auto* Statement = new(Allocator) AST::FSelectionStatement(Allocator, Parser.Scanner.GetCurrentToken()->SourceInfo);
 
@@ -2371,8 +2382,8 @@ Done:
 	FHlslParser::FHlslParser(FLinearAllocator* InAllocator, FCompilerMessages& InCompilerMessages, TConstArrayView<FScopedDeclarations> InScopedDeclarations) :
 		Scanner(InCompilerMessages),
 		CompilerMessages(InCompilerMessages),
-		GlobalScope(InAllocator, nullptr),
-		Namespaces(InAllocator, nullptr),
+		GlobalScope(InAllocator, nullptr, FSymbolScope::EType::Global),
+		Namespaces(InAllocator, nullptr, FSymbolScope::EType::Namespace),
 		Allocator(InAllocator)
 	{
 		CurrentScope = &GlobalScope;
@@ -2384,7 +2395,7 @@ Done:
 			FSymbolScope* IterScope = CurrentScope;
 			for (FStringView Scope : ScopedDeclarations.Scope)
 			{
-				ScopeStack.Emplace(Allocator, &IterScope);
+				ScopeStack.Emplace(Allocator, &IterScope, FSymbolScope::EType::ScopedDeclaration);
 				IterScope->Name = Scope.GetData();
 			}
 
diff --git a/Engine/Source/Developer/ShaderCompilerCommon/Private/HlslUtils.cpp b/Engine/Source/Developer/ShaderCompilerCommon/Private/HlslUtils.cpp
--- a/Engine/Source/Developer/ShaderCompilerCommon/Private/HlslUtils.cpp
+++ b/Engine/Source/Developer/ShaderCompilerCommon/Private/HlslUtils.cpp
@@ -2317,10 +2317,10 @@ TEST_CASE_NAMED(FShaderRenameLocalsTest, "Shaders::RenameLocals", "[EditorContex
 {
 	FString ShaderCodeString(
 		TEXT(R"(
-float fun(float Param_A, float Param_B)
+float fun(float Param_A, float Param_B, float Param_C)
 {
 	float C = 123.0;
-	return Param_B * Param_A + C;
+	return Param_B * Param_A + C * Param_C;
 }
 float4 main(float4 Param_C : TEXCOORD0) : SV_Target
 {
@@ -2349,6 +2349,10 @@ float4 main(float4 Param_C : TEXCOORD0) : SV_Target
 
 	return float4(A, B, A+B, fun(A,B)) + Param_C;
 }
+float4 bar(float Param_A, float Param_B, float4 Param_C)
+{
+	return Param_C * Param_A + Param_B;
+}
 )"));
 
 	FString EntryPoint(TEXT("main"));
@@ -2458,6 +2462,8 @@ float4 main(float4 Param_C : TEXCOORD0) : SV_Target
 	TEST_TRUE(TEXT("Function uses renamed for loop variable I"), Result.Contains(TEXT("for (int Renamed_I = 0; Renamed_I < 4; ++Renamed_I)")));
 	TEST_TRUE(TEXT("Function uses renamed for loop variables J and K"), Result.Contains(TEXT("for (int Renamed_J = 0, Renamed_K = 1; Renamed_J < 4; ++Renamed_J)")));
 	TEST_TRUE(TEXT("Function uses renamed locals"), Result.Contains(TEXT("float4(Renamed_A, Renamed_B, (Renamed_A + Renamed_B), fun(Renamed_A, Renamed_B)) + Renamed_Param_C")));
+	TEST_TRUE(TEXT("Function foo() signature is expected to be unchanged"), Result.Contains(TEXT("fun(float Param_A, float Param_B, float Param_C)")));
+	TEST_TRUE(TEXT("Function bar() signature is expected to be unchanged"), Result.Contains(TEXT("float4 bar(float Param_A, float Param_B, float4 Param_C)")));
 }
 
 #endif // WITH_LOW_LEVEL_TESTS
