Add initial implementation of AST-level shader function inlining utility

This serves as a base for the further work:
- Inlining from nested compound statement scopes will be supported (currently only calls from the top level are handled)
- Functions with single trivial return can be supported
- Multiple levels of inlining can be supported
- Functions or call sites can be annotated with a custom attribute for AST-level inlining opt-in

diff --git a/Engine/Source/Developer/ShaderCompilerCommon/Private/HlslUtils.cpp b/Engine/Source/Developer/ShaderCompilerCommon/Private/HlslUtils.cpp
--- a/Engine/Source/Developer/ShaderCompilerCommon/Private/HlslUtils.cpp
+++ b/Engine/Source/Developer/ShaderCompilerCommon/Private/HlslUtils.cpp
@@ -2110,6 +2110,361 @@ namespace UE::ShaderCompilerCommon
 	}
 }
 
+static void VisitCompoundStatementChildren(
+	CrossCompiler::AST::FCompoundStatement* InRoot,
+	TFunction<void(CrossCompiler::AST::FNode* Node)> Visitor)
+{
+	using namespace CrossCompiler::AST;
+
+	TArray<FCompoundStatement*, TInlineAllocator<128>> Stack;
+	Stack.Push(InRoot);
+
+	auto VisitAndConditionallyPushNode = [&Stack, &Visitor](FNode* Node)
+	{
+		if (!Node)
+		{
+			return false;
+		}
+
+		Visitor(Node);
+
+		if (FCompoundStatement* CompoundStatement = Node->AsCompoundStatement())
+		{
+			Stack.Push(CompoundStatement);
+			return true;
+		}
+		else
+		{
+			return false;
+		}
+	};
+
+	while (!Stack.IsEmpty())
+	{
+		FCompoundStatement* Node = Stack.Last();
+		Stack.Pop();
+
+		for (FNode* ChildStatement : Node->Statements)
+		{
+			if (VisitAndConditionallyPushNode(ChildStatement))
+			{
+				continue;
+			}
+
+			if (FSelectionStatement* SelectionStatement = ChildStatement->AsSelectionStatement())
+			{
+				while (SelectionStatement)
+				{
+					VisitAndConditionallyPushNode(SelectionStatement->ThenStatement);
+					VisitAndConditionallyPushNode(SelectionStatement->ElseStatement);
+
+					// Handle chain of `else if` statements
+					SelectionStatement = SelectionStatement->ElseStatement ? SelectionStatement->ElseStatement->AsSelectionStatement() : nullptr;
+				}
+			}
+			else if (FSwitchStatement* SwitchStatement = ChildStatement->AsSwitchStatement())
+			{
+				if (SwitchStatement->Body && SwitchStatement->Body->CaseList)
+				{
+					for (FCaseStatement* CaseStatement : SwitchStatement->Body->CaseList->Cases)
+					{
+						for (FNode* CaseStatementChildNode : CaseStatement->Statements)
+						{
+							VisitAndConditionallyPushNode(CaseStatementChildNode);
+						}
+					}
+				}
+			}
+			else if (FIterationStatement* IterationStatement = ChildStatement->AsIterationStatement())
+			{
+				VisitAndConditionallyPushNode(IterationStatement->Body);
+			}
+		}
+	}
+}
+
+static void VisitFunctionIdentifiers(
+	CrossCompiler::AST::FFunctionDefinition* FunctionDefinition,
+	TFunction<void(CrossCompiler::AST::FIdentifier* Identifier)> ProcessIdentifier)
+{
+	using namespace CrossCompiler::AST;
+
+	auto ProcessDeclaration = [&ProcessIdentifier](FDeclaration* DeclarationNode)
+	{
+		if (!DeclarationNode || !DeclarationNode->Identifier)
+		{
+			return;
+		}
+
+		ProcessIdentifier(DeclarationNode->Identifier);
+	};
+
+	auto ProcessDeclaratorList = [&ProcessDeclaration](FDeclaratorList* DeclaratorListNode)
+	{
+		for (FNode* InnerDeclarationNode : DeclaratorListNode->Declarations)
+		{
+			if (FDeclaration* InnerDeclaration = InnerDeclarationNode->AsDeclaration())
+			{
+				ProcessDeclaration(InnerDeclaration);
+			}
+		}
+	};
+
+	for (FNode* ParameterNode : FunctionDefinition->Prototype->Parameters)
+	{
+		FParameterDeclarator* ParameterDeclarator = ParameterNode->AsParameterDeclarator();
+		if (!ParameterDeclarator)
+		{
+			continue;
+		}
+
+		if (ParameterDeclarator->Identifier)
+		{
+			ProcessIdentifier(ParameterDeclarator->Identifier);
+		}
+	}
+
+	VisitCompoundStatementChildren(FunctionDefinition->Body, [&ProcessDeclaratorList, &ProcessDeclaration](FNode* StatementNode)
+	{
+		if (FIterationStatement* IterationStatement = StatementNode->AsIterationStatement())
+		{
+			FNode* InitStatementNode = IterationStatement->InitStatement;
+			if (FDeclaratorList* DeclaratorListNode = InitStatementNode->AsDeclaratorList())
+			{
+				ProcessDeclaratorList(DeclaratorListNode);
+			}
+		}
+		else if (FDeclaration* Declaration = StatementNode->AsDeclaration())
+		{
+			ProcessDeclaration(Declaration);
+		}
+		else if (FDeclaratorList* DeclaratorListNode = StatementNode->AsDeclaratorList())
+		{
+			ProcessDeclaratorList(DeclaratorListNode);
+		}
+	});
+}
+
+static FString FormatSourceMessage(const CrossCompiler::FSourceInfo& SourceInfo, const TCHAR* Message)
+{
+	const TCHAR* FileName = (SourceInfo.Filename && !SourceInfo.Filename->IsEmpty()) ? **SourceInfo.Filename : TEXT("<unknown>");
+	return FString::Printf(TEXT("%s(%d,%d): %s\n"), FileName, SourceInfo.Line, SourceInfo.Column, Message);
+}
+
+bool InlineFunction(FString& InOutSourceCode, FString& InOutEntryPoint, const FStringView FunctionToInline, TArray<FString>& OutErrors)
+{
+	using namespace CrossCompiler::AST;
+
+	CrossCompiler::FNodeContainer ShaderNodes;
+
+	CrossCompiler::FCompilerMessages CompilerMessages;
+	TArray<FScopedDeclarations> ScopedDeclarations;
+	if (!CrossCompiler::Parser::Parse(InOutSourceCode, TEXT("inliner.hlsl"), CompilerMessages, ScopedDeclarations, ShaderNodes))
+	{
+		for (const auto& Message : CompilerMessages.MessageList)
+		{
+			if (Message.bIsError)
+			{
+				OutErrors.Add(Message.Message);
+			}
+		}
+		return false;
+	}
+
+	auto FindFirstFunctionByName = [&ShaderNodes](FStringView Name) -> FFunctionDefinition*
+		{
+			for (FNode* Node : ShaderNodes.Nodes)
+			{
+				FFunctionDefinition* FunctionDefinition = Node->AsFunctionDefinition();
+				if (FunctionDefinition && FunctionDefinition->Body && FunctionDefinition->Prototype->Identifier->Equals(Name))
+				{
+					return FunctionDefinition;
+				}
+			}
+
+			return nullptr;
+		};
+
+	FFunctionDefinition* EntryFunctionDefinition = FindFirstFunctionByName(InOutEntryPoint);
+	if (!EntryFunctionDefinition)
+	{
+		OutErrors.Add(TEXT("Could not find entry point function definition"));
+		return false;
+	}
+
+	FFunctionDefinition* InlinedFunctionDefinition = FindFirstFunctionByName(FunctionToInline);
+	if (!InlinedFunctionDefinition || !InlinedFunctionDefinition->Body)
+	{
+		OutErrors.Add(TEXT("Could not find inlined function definition"));
+		return false;
+	}
+
+	// Check that the inlined function meets the current big limitation: no return statements.
+
+	bool bInlinedFunctionCompatible = true;
+	VisitCompoundStatementChildren(InlinedFunctionDefinition->Body, [&bInlinedFunctionCompatible, &OutErrors] (FNode* Node)
+	{
+		if (FJumpStatement* JumpStatement = Node->AsJumpStatement())
+		{
+			if (JumpStatement->Type == EJumpType::Return)
+			{
+				OutErrors.Add(FormatSourceMessage(Node->SourceInfo, TEXT("Inlined functions may not have return statements. Out parameters must be used instead.")));
+				bInlinedFunctionCompatible = false;
+			}
+		}
+	});
+
+	if (!bInlinedFunctionCompatible)
+	{
+		return false;
+	}
+	
+	auto AddPrefixToIdentifier = [&ShaderNodes](FIdentifier* Identifier)
+	{
+		TStringBuilder<256> NewName;
+		NewName << TEXT("INLINE_");
+		NewName << Identifier->ToStringView();
+		Identifier->Rename(&ShaderNodes.Allocator, NewName);
+	};
+
+	// TODO: undo the rename if we need to inline the function multiple times
+	// TODO: could only rename colliding names
+	VisitFunctionIdentifiers(InlinedFunctionDefinition, AddPrefixToIdentifier);
+
+	CrossCompiler::FNodeContainer OutputNodes;
+
+	auto AsCallToInlinedFunction = [&FunctionToInline](FNode* Statement) -> FExpressionStatement*
+		{
+			if (FExpressionStatement* ExpressionStatement = Statement->AsExpressionStatement())
+			{
+				if (ExpressionStatement->Expression->Operator == EOperators::FunctionCall)
+				{
+					FFunctionExpression* FunctionExpression = static_cast<FFunctionExpression*>(ExpressionStatement->Expression);
+					if (FunctionExpression->Callee->Identifier->Equals(FunctionToInline))
+					{
+						return ExpressionStatement;
+					}
+				}
+			}
+
+			return nullptr;
+		};
+
+	const CrossCompiler::FSourceInfo SourceInfo; // todo: more accurate/better one?
+
+	FFunctionDefinition* OutputFunction = OutputNodes.AllocNode<FFunctionDefinition>(SourceInfo);
+	FCompoundStatement* OutputBody = OutputNodes.AllocNode<FCompoundStatement>(SourceInfo);
+
+	TStringBuilder<256> OptimizedEntryPoint;
+	OptimizedEntryPoint << EntryFunctionDefinition->Prototype->Identifier->ToStringView();
+	OptimizedEntryPoint << TEXTVIEW("__INLINED");
+
+	OutputFunction->Prototype = EntryFunctionDefinition->Prototype;
+	OutputFunction->Prototype->Identifier = FIdentifier::New(&OutputNodes.Allocator, OptimizedEntryPoint);
+	OutputFunction->Body = OutputBody;
+
+	// copy over original statements until we hit a call to the function we're inlining
+
+	// TODO: recursively process compound statements
+	for (FNode* Statement : EntryFunctionDefinition->Body->Statements)
+	{
+		FNode* StatementToAdd = Statement; // shallow clone of the original statement by default
+
+		if (FExpressionStatement* CallExpressionStatement = AsCallToInlinedFunction(Statement))
+		{
+			FCompoundStatement* InlinedBody = OutputNodes.AllocNode<FCompoundStatement>(SourceInfo);
+
+			const FExpression* CallExpression = CallExpressionStatement->Expression;
+
+			// add declarations for function paramters
+			const int32 NumParemeters = InlinedFunctionDefinition->Prototype->Parameters.Num();
+			const int32 NumProvidedParameters = CallExpression ? CallExpression->Expressions.Num() : 0;
+
+			struct FOutParameter
+			{
+				FParameterDeclarator* Declarator = nullptr;
+				FExpression* Expression = nullptr;
+			};
+
+			TArray<FOutParameter> OutParameters;
+
+			for (int32 ParamIndex = 0; ParamIndex < NumParemeters; ++ParamIndex)
+			{
+				FNode* FunParamNode = InlinedFunctionDefinition->Prototype->Parameters[ParamIndex];
+				FParameterDeclarator* FunParam = FunParamNode->AsParameterDeclarator();
+				check(FunParam);
+
+				FDeclaration* LocalDeclaration = OutputNodes.AllocNode<FDeclaration>(SourceInfo);
+
+				LocalDeclaration->Identifier = FunParam->Identifier;
+				LocalDeclaration->Semantic = FunParam->Semantic;
+				LocalDeclaration->bIsArray = FunParam->bIsArray;
+				LocalDeclaration->ArraySize = FunParam->ArraySize;
+
+				const bool bExclusiveOutParam = FunParam->Type->Qualifier.bOut && !FunParam->Type->Qualifier.bIn;
+
+				FExpression* InitializerExpression = ParamIndex < NumProvidedParameters
+					? CallExpression->Expressions[ParamIndex]
+					: LocalDeclaration->Initializer = FunParam->DefaultValue;
+
+					if (!bExclusiveOutParam)
+					{
+						LocalDeclaration->Initializer = InitializerExpression;
+					}
+
+					FDeclaratorList* LocalDeclaratorList = OutputNodes.AllocNode<FDeclaratorList>(SourceInfo);
+					LocalDeclaratorList->Type = OutputNodes.AllocNode<FFullySpecifiedType>(SourceInfo);
+					LocalDeclaratorList->Type->Qualifier = FunParam->Type->Qualifier;
+					// Clear in/out qualifiers when declaring the local
+					LocalDeclaratorList->Type->Qualifier.bIn = 0;
+					LocalDeclaratorList->Type->Qualifier.bOut = 0;
+					LocalDeclaratorList->Type->Specifier = FunParam->Type->Specifier; // shallow clone
+
+					LocalDeclaratorList->Declarations.Add(LocalDeclaration);
+
+					InlinedBody->Statements.Add(LocalDeclaratorList);
+
+					// Save the parameter description
+					if (FunParam->Type->Qualifier.bOut)
+					{
+						FOutParameter OutParam;
+						OutParam.Declarator = FunParam;
+						OutParam.Expression = InitializerExpression;
+						OutParameters.Add(OutParam);
+					}
+			}
+
+			InlinedBody->Statements.Append(InlinedFunctionDefinition->Body->Statements);
+
+			for (const FOutParameter& OutParam : OutParameters)
+			{
+				FExpression* RHS = OutputNodes.AllocNode<FExpression>(EOperators::Identifier, SourceInfo);
+				RHS->Identifier = OutParam.Declarator->Identifier;
+				FBinaryExpression* AssignmentExpression = OutputNodes.AllocNode<FBinaryExpression>(EOperators::Assign, OutParam.Expression, RHS, SourceInfo);
+				FExpressionStatement* AssignmentStatement = OutputNodes.AllocNode<FExpressionStatement>(AssignmentExpression, SourceInfo);
+				InlinedBody->Statements.Add(AssignmentStatement);
+			}
+
+			StatementToAdd = InlinedBody;
+		}
+
+		OutputFunction->Body->Statements.Add(StatementToAdd);
+	}
+
+	OutputNodes.Nodes.Add(OutputFunction);
+
+	FASTWriter GeneratedCodeWriter(InOutSourceCode);
+
+	if (!OutputNodes.Nodes.IsEmpty())
+	{
+		OutputNodes.Nodes.Last()->Write(GeneratedCodeWriter);
+	}
+
+	InOutEntryPoint = OptimizedEntryPoint;
+
+	return true;
+}
+
 #if WITH_LOW_LEVEL_TESTS
 
 static void GenerateShaderCode(const CrossCompiler::TLinearArray<CrossCompiler::AST::FNode*>& ASTNodes, FString& OutResult)
@@ -2375,27 +2730,6 @@ float4 bar(float Param_A, float Param_B, float4 Param_C)
 				Identifier->Rename(Allocator, NewName);
 			};
 
-			auto ProcessDeclaration = [&ProcessIdentifier](FDeclaration* DeclarationNode)
-				{
-					if (!DeclarationNode || !DeclarationNode->Identifier)
-					{
-						return;
-					}
-
-					ProcessIdentifier(DeclarationNode->Identifier);
-				};
-
-			auto ProcessDeclaratorList = [&ProcessDeclaration](FDeclaratorList* DeclaratorListNode)
-			{
-				for (FNode* InnerDeclarationNode : DeclaratorListNode->Declarations)
-				{
-					if (FDeclaration* InnerDeclaration = InnerDeclarationNode->AsDeclaration())
-					{
-						ProcessDeclaration(InnerDeclaration);
-					}
-				}
-			};
-
 			for (FNode* Node : ASTNodes)
 			{
 				FFunctionDefinition* FunctionDefinition = Node->AsFunctionDefinition();
@@ -2404,52 +2738,7 @@ float4 bar(float Param_A, float Param_B, float4 Param_C)
 					continue;
 				}
 
-				for (FNode* ParameterNode : FunctionDefinition->Prototype->Parameters)
-				{
-					FParameterDeclarator* ParameterDeclarator = ParameterNode->AsParameterDeclarator();
-					if (!ParameterDeclarator)
-					{
-						continue;
-					}
-
-					if (ParameterDeclarator->Identifier)
-					{
-						ProcessIdentifier(ParameterDeclarator->Identifier);
-					}
-				}
-
-				TArray<FCompoundStatement*> Stack;
-				Stack.Push(FunctionDefinition->Body);
-
-				while (!Stack.IsEmpty())
-				{
-					FCompoundStatement* ParentCompoundStatementNode = Stack.Last();
-					Stack.Pop();
-
-					for (FNode* StatementNode : ParentCompoundStatementNode->Statements)
-					{
-						if (FCompoundStatement* CompoundStatement = StatementNode->AsCompoundStatement())
-						{
-							Stack.Push(CompoundStatement);
-						}
-						else if (FIterationStatement* IterationStatement = StatementNode->AsIterationStatement())
-						{
-							FNode* InitStatementNode = IterationStatement->InitStatement;
-							if (FDeclaratorList* DeclaratorListNode = InitStatementNode->AsDeclaratorList())
-							{
-								ProcessDeclaratorList(DeclaratorListNode);
-							}
-						}
-						else if (FDeclaration* Declaration = StatementNode->AsDeclaration())
-						{
-							ProcessDeclaration(Declaration);
-						}
-						else if (FDeclaratorList* DeclaratorListNode = StatementNode->AsDeclaratorList())
-						{
-							ProcessDeclaratorList(DeclaratorListNode);
-						}
-					}
-				}
+				VisitFunctionIdentifiers(FunctionDefinition, ProcessIdentifier);
 			}
 
 			GenerateShaderCode(ASTNodes, Result);
@@ -2466,4 +2755,87 @@ float4 bar(float Param_A, float Param_B, float4 Param_C)
 	TEST_TRUE(TEXT("Function bar() signature is expected to be unchanged"), Result.Contains(TEXT("float4 bar(float Param_A, float Param_B, float4 Param_C)")));
 }
 
+TEST_CASE_NAMED(FShaderInlineFunctionTest, "Shaders::InlineFunction", "[EditorContext][EngineFilter]")
+{
+	FString ShaderCodeString(
+		TEXT(R"(// dxc -T ps_6_0 inline.hlsl
+struct FParams
+{
+	float a;
+	float b;
+};
+
+void foo(out float4 res, in FParams params, float a, float b = 0.0, int q=1)
+{
+	if (b > 10.0)
+	{
+		a *= b;
+	}
+
+	if (b < -10.0)
+		b *= 2;
+
+	float c = (a + b + params.a) * params.b;
+	c = c + 1.0;
+	res = float4(a, b, c, 1.0);
+
+	for (int i=0; i<10; ++i)
+	{
+		int j = i * 2;
+		a += 0.1 + (float)j;
+	}
+
+	if (b < 0)
+	{
+		int x = 1;
+		res.x += 1.0;
+	}
+
+	switch (q)
+	{
+		case 0:
+			break;
+		case 1:
+		case 2:
+		{
+			res.x += 1.0;
+			break;
+		}
+		case 3:
+		{
+			int w = q;
+			res.x += (float)w;
+			break;
+		}
+	}
+}
+
+float4 main() : SV_Target
+{
+	float x = 5.0;
+	float c = 3.0;
+	float4 res;
+	FParams params;
+	params.a = 1.5f;
+	params.b = 2.5f;
+	foo(res, params, 1.23f);
+	res.x += x * c;
+	return res;
+})"));
+
+	FString EntryPointString(TEXT("main"));
+	FString FunctionNameString(TEXT("foo"));
+
+	TArray<FString> Errors;
+	const bool bSuccess = InlineFunction(ShaderCodeString, EntryPointString, FunctionNameString, Errors);
+
+	TEST_TRUE(TEXT("Function inlining succeeded"), bSuccess);
+	TEST_TRUE(TEXT("Function arguments are handled"), ShaderCodeString.Contains(TEXT("float4 INLINE_res;")));
+	TEST_TRUE(TEXT("Function arguments with local parameters are handled"), ShaderCodeString.Contains(TEXT("FParams INLINE_params = params;")));
+	TEST_TRUE(TEXT("Function arguments with local parameters are handled"), ShaderCodeString.Contains(TEXT("float INLINE_a = 1.23f;")));
+	TEST_TRUE(TEXT("Function arguments with default values are handled"), ShaderCodeString.Contains(TEXT("float INLINE_b = 0.0;")));
+	TEST_TRUE(TEXT("Function arguments with default values are handled"), ShaderCodeString.Contains(TEXT("int INLINE_q = 1;")));
+	TEST_TRUE(TEXT("Out parameter is handled"), ShaderCodeString.Contains(TEXT("res = INLINE_res")));
+}
+
 #endif // WITH_LOW_LEVEL_TESTS
diff --git a/Engine/Source/Developer/ShaderCompilerCommon/Private/HlslUtils.h b/Engine/Source/Developer/ShaderCompilerCommon/Private/HlslUtils.h
--- a/Engine/Source/Developer/ShaderCompilerCommon/Private/HlslUtils.h
+++ b/Engine/Source/Developer/ShaderCompilerCommon/Private/HlslUtils.h
@@ -16,6 +16,8 @@ namespace CrossCompiler
 	namespace AST
 	{
 		class FNode;
+		struct FFunctionDefinition;
+		class FIdentifier;
 	}
 
 	namespace Memory
@@ -235,6 +237,14 @@ namespace CrossCompiler
 		{
 		}
 
+		template <typename T, typename... ArgsType>
+		T* AllocNode(ArgsType&&... Args)
+		{
+			static_assert(std::is_base_of_v<CrossCompiler::AST::FNode, T>, "Type must inherit from CrossCompiler::AST::FNode");
+			T* Result = new(&Allocator) T(&Allocator, Forward<ArgsType>(Args)...);
+			return Result;
+		}
+
 		FLinearAllocator Allocator;
 		TLinearArray<CrossCompiler::AST::FNode*> Nodes;
 	};
diff --git a/Engine/Source/Developer/ShaderCompilerCommon/Public/ShaderCompilerCommon.h b/Engine/Source/Developer/ShaderCompilerCommon/Public/ShaderCompilerCommon.h
--- a/Engine/Source/Developer/ShaderCompilerCommon/Public/ShaderCompilerCommon.h
+++ b/Engine/Source/Developer/ShaderCompilerCommon/Public/ShaderCompilerCommon.h
@@ -422,6 +422,13 @@ extern SHADERCOMPILERCOMMON_API bool FindEntryPointParameters(
 
 extern SHADERCOMPILERCOMMON_API bool ConvertFromFP32ToFP16(FString& InOutSourceCode, TArray<FString>& OutErrors);
 
+// EXPERIMENTAL: Adds a new function to the shader where calls to `FunctionToInline` from `EntryPoint` are inlined using AST manipulation
+// Current limitations:
+// - Inlining from nested compound statement scopes is not supported (currently only trivial call statements in the body are processed)
+// - Functions with return statements cannot be inlined
+// - Multiple levels of inlining are not supported
+extern SHADERCOMPILERCOMMON_API bool InlineFunction(FString& InOutSourceCode, FString& InOutEntryPoint, const FStringView FunctionToInline, TArray<FString>& OutErrors);
+
 enum class EShaderConductorTarget
 {
 	Dxil,
