Metal MSC - Fix large metallib size by only embedding debug information when needed with r.Shaders.Symbols=1
The debug information is already embedded automatically when needed with -Zi & -Qembed_debug in MetalCompileShaderMSC.cpp and nothing else is needed.

#tests Local test with and without symbols in editor and packaged game

diff --git a/Engine/Source/Developer/ShaderCompilerCommon/Private/ShaderConductorContext.cpp b/Engine/Source/Developer/ShaderCompilerCommon/Private/ShaderConductorContext.cpp
--- a/Engine/Source/Developer/ShaderCompilerCommon/Private/ShaderConductorContext.cpp
+++ b/Engine/Source/Developer/ShaderCompilerCommon/Private/ShaderConductorContext.cpp
@@ -683,9 +683,6 @@ namespace CrossCompiler
 		
         ShaderConductor::Compiler::Options ScOptions;
         ConvertScOptions(*Intermediates, Options, ScOptions, false, bGenerateSpirv);
-        
-        // Force PDB generation (required to generate DXIL reflection).
-        ScOptions.enableDebugInfo = true;
 
         // Compile HLSL source code to DXIL
         bool bSucceeded = false;
