Optimize some of the temporary string allocations during shader parsing

* Replace temporary FString-s with TStringBuilder-s in a few trivial places.
* Replace AddToken with EmplaceToken that takes FHlslToken by move to avoid copying the string it contains

diff --git a/Engine/Source/Developer/ShaderCompilerCommon/Private/HlslExpressionParser.inl b/Engine/Source/Developer/ShaderCompilerCommon/Private/HlslExpressionParser.inl
--- a/Engine/Source/Developer/ShaderCompilerCommon/Private/HlslExpressionParser.inl
+++ b/Engine/Source/Developer/ShaderCompilerCommon/Private/HlslExpressionParser.inl
@@ -514,16 +514,16 @@ namespace CrossCompiler
 		if (bMatched)
 		{
 			//#todo-rco: Don't re-allocate types
-			FString TypeName;
+			TStringBuilder<64> TypeName;
 			TypeName += (ExtraQualifierFlags & EEQ_SNORM) == EEQ_SNORM
-				? TEXT("snorm ")
+				? TEXTVIEW("snorm ")
 				: ((ExtraQualifierFlags & EEQ_UNORM) == EEQ_UNORM
-					? TEXT("unorm ")
-					: TEXT(""));
+					? TEXTVIEW("unorm ")
+					: TEXTVIEW(""));
 			TypeName += Token->String;
 
 			auto* Type = new(Allocator) AST::FTypeSpecifier(Allocator, Token->SourceInfo);
-			Type->TypeName = Allocator->Strdup(TypeName);
+			Type->TypeName = Allocator->Strdup(TypeName.ToView());
 			Type->InnerType = InnerType;
 			Type->bPrecise = (ExtraQualifierFlags & EEQ_PRECISE) == EEQ_PRECISE;
 			*OutSpecifier = Type;
@@ -575,7 +575,7 @@ namespace CrossCompiler
 		{
 			auto* Token1 = Scanner.PeekToken(1);
 			auto* Token2 = Scanner.PeekToken(2);
-			FString TypeString;
+			TStringBuilder<256> TypeString;
 			if (Token && Token->Token == EHlslToken::Identifier && Token1 && Token1->Token == EHlslToken::ColonColon && Token2 && Token2->Token == EHlslToken::Identifier && SymbolScope)
 			{
 				auto* Namespace = SymbolScope->GetGlobalScope();
@@ -613,7 +613,7 @@ namespace CrossCompiler
 							if (FSymbolScope::FindType(Namespace, InnerOrType->String, false))
 							{
 								auto* Type = new(Allocator) AST::FTypeSpecifier(Allocator, Token->SourceInfo);
-								Type->TypeName = Allocator->Strdup(*TypeString);
+								Type->TypeName = Allocator->Strdup(TypeString.ToView());
 								*OutSpecifier = Type;
 								return EParseResult::Matched;
 							}
diff --git a/Engine/Source/Developer/ShaderCompilerCommon/Private/HlslLexer.cpp b/Engine/Source/Developer/ShaderCompilerCommon/Private/HlslLexer.cpp
--- a/Engine/Source/Developer/ShaderCompilerCommon/Private/HlslLexer.cpp
+++ b/Engine/Source/Developer/ShaderCompilerCommon/Private/HlslLexer.cpp
@@ -79,7 +79,7 @@ namespace CrossCompiler
 		}
 	}
 
-	static bool MatchSymbolToken(const TCHAR* InString, const TCHAR** OutString, EHlslToken& OutToken, FString* OutTokenString, bool bGreedy)
+	static bool MatchSymbolToken(const TCHAR* InString, const TCHAR** OutString, EHlslToken& OutToken, TStringBuilderBase<TCHAR>* OutTokenString, bool bGreedy)
 	{
 		const TCHAR* OriginalString = InString;
 		FKeywordToken* Found = Keywords.Find(*InString);
@@ -111,8 +111,8 @@ namespace CrossCompiler
 							OutToken = Found->Current;
 							if (OutTokenString)
 							{
-								*OutTokenString = TEXT("");
-								OutTokenString->AppendChars(OriginalString, InString - OriginalString);
+								OutTokenString->Reset();
+								OutTokenString->Append(OriginalString, InString - OriginalString);
 							}
 
 							if (OutString)
@@ -136,8 +136,8 @@ namespace CrossCompiler
 				OutToken = Found->Current;
 				if (OutTokenString)
 				{
-					*OutTokenString = TEXT("");
-					OutTokenString->AppendChars(OriginalString, InString - OriginalString);
+					OutTokenString->Reset();
+					OutTokenString->Append(OriginalString, InString - OriginalString);
 				}
 
 				if (OutString)
@@ -613,6 +613,11 @@ namespace CrossCompiler
 			CurrentLineStart = Current;
 		}
 
+		bool MatchString(FStringView String)
+		{
+			return MatchString(String.GetData(), String.Len());
+		}
+
 		bool MatchString(const TCHAR* Target, int32 TargetLen)
 		{
 			if (Current + TargetLen <= End)
@@ -672,7 +677,7 @@ namespace CrossCompiler
 			}
 		}
 
-		bool MatchLiteralFloat(FString& OutLiteral, ELiteralType& OutType)
+		bool MatchLiteralFloat(TStringBuilderBase<TCHAR>& OutLiteral, ELiteralType& OutType)
 		{
 			auto* Original = Current;
 			TCHAR Char = Peek();
@@ -775,7 +780,7 @@ namespace CrossCompiler
 			}
 
 		Done:
-			OutLiteral = FString::ConstructFromPtrSize(Original, static_cast<int32>(Current - Original));
+			OutLiteral.Append(Original, static_cast<int32>(Current - Original));
 			OutType = Type;
 			return true;
 
@@ -784,7 +789,7 @@ namespace CrossCompiler
 			return false;
 		}
 
-		bool MatchQuotedString(FString& OutString)
+		bool MatchQuotedString(TStringBuilderBase<TCHAR>& OutString)
 		{
 			if (!Match('"'))
 			{
@@ -823,7 +828,7 @@ namespace CrossCompiler
 			return false;
 		}
 
-		bool MatchIdentifier(FString& OutIdentifier)
+		bool MatchIdentifier(TStringBuilderBase<TCHAR>& OutIdentifier)
 		{
 			if (HasCharsAvailable())
 			{
@@ -853,7 +858,7 @@ namespace CrossCompiler
 			return false;
 		}
 
-		bool MatchSymbol(EHlslToken& OutToken, FString& OutTokenString)
+		bool MatchSymbol(EHlslToken& OutToken, TStringBuilderBase<TCHAR>& OutTokenString)
 		{
 			if (HasCharsAvailable())
 			{
@@ -868,9 +873,8 @@ namespace CrossCompiler
 
 		static void ProcessDirective(FTokenizer& Tokenizer, FCompilerMessages& CompilerMessages, class FHlslScanner& Scanner);
 
-		FString ReadToEndOfLine(bool bSkipToNextLine = true)
+		void ReadToEndOfLine(TStringBuilderBase<TCHAR>& OutString, bool bSkipToNextLine = true)
 		{
-			FString String;
 			const TCHAR* Start = Current;
 			const TCHAR* EndOfLine = Current;
 
@@ -898,11 +902,10 @@ namespace CrossCompiler
 			}
 
 			int32 Count = (int32)(EndOfLine - Start) + 1;
-			String.AppendChars(Start, Count);
-			return String;
+			OutString.Append(Start, Count);
 		}
 
-		bool RuleDecimalInteger(FString& OutValue)
+		bool RuleDecimalInteger(TStringBuilderBase<TCHAR>& OutValue)
 		{
 			// [1-9][0-9]*
 			auto Char = Peek();
@@ -930,7 +933,7 @@ namespace CrossCompiler
 
 		bool RuleDecimalIntegerValue(uint32& OutValue)
 		{
-			FString String;
+			TStringBuilder<11> String; // large enough for uint32 max + terminator
 			if (RuleDecimalInteger(String))
 			{
 				OutValue = (uint32)FCString::Atoi(*String);
@@ -940,7 +943,7 @@ namespace CrossCompiler
 			return false;
 		}
 
-		bool RuleOctalInteger(FString& OutValue)
+		bool RuleOctalInteger(TStringBuilderBase<TCHAR>& OutValue)
 		{
 			// 0[0-7]*
 			auto Char = Peek();
@@ -968,7 +971,7 @@ namespace CrossCompiler
 			return true;
 		}
 
-		bool RuleHexadecimalInteger(FString& OutValue)
+		bool RuleHexadecimalInteger(TStringBuilderBase<TCHAR>& OutValue)
 		{
 			// 0[xX][0-9a-zA-Z]+
 			auto Char = Peek();
@@ -1030,7 +1033,7 @@ namespace CrossCompiler
 			return false;
 		}
 
-		bool MatchLiteral(FString& OutValue, ELiteralType& OutLiteralType)
+		bool MatchLiteral(TStringBuilderBase<TCHAR>& OutValue, ELiteralType& OutLiteralType)
 		{
 			if (MatchLiteralFloat(OutValue, OutLiteralType))
 			{
@@ -1042,7 +1045,7 @@ namespace CrossCompiler
 			}
 		}
 
-		bool MatchLiteralInteger(FString& OutValue, ELiteralType& OutLiteralType)
+		bool MatchLiteralInteger(TStringBuilderBase<TCHAR>& OutValue, ELiteralType& OutLiteralType)
 		{
 			if (RuleHexadecimalInteger(OutValue))
 			{
@@ -1061,7 +1064,7 @@ namespace CrossCompiler
 			}
 			else if (RuleOctalInteger(OutValue))
 			{
-				if (OutValue.Len() == 1 && OutValue[0] == (TCHAR)'0')
+				if (OutValue.Len() == 1 && OutValue.LastChar() == (TCHAR)'0')
 				{
 					auto Char = Peek();
 					if (IsValidIntegerSuffix(Char))
@@ -1111,9 +1114,9 @@ namespace CrossCompiler
 	{
 	}
 
-	inline void FHlslScanner::AddToken(const FHlslToken& Token, const FTokenizer& Tokenizer)
+	inline void FHlslScanner::EmplaceToken(FHlslToken&& Token, const FTokenizer& Tokenizer)
 	{
-		int32 TokenIndex = Tokens.Add(Token);
+		int32 TokenIndex = Tokens.Emplace(MoveTemp(Token));
 		Tokens[TokenIndex].SourceInfo.Filename = &SourceFilenames.Last();
 		Tokens[TokenIndex].SourceInfo.Line = Tokenizer.Line;
 		Tokens[TokenIndex].SourceInfo.Column = (int32)(Tokenizer.Current - Tokenizer.CurrentLineStart) + 1;
@@ -1147,35 +1150,35 @@ namespace CrossCompiler
 			}
 			else
 			{
-				FString Identifier;
+				TStringBuilder<128> Identifier;
 				EHlslToken SymbolToken;
 				ELiteralType LiteralType = ELiteralType::Unknown;
 				if (Tokenizer.MatchLiteral(Identifier, LiteralType))
 				{
-					AddToken(FHlslToken(Identifier, LiteralType), Tokenizer);
+					EmplaceToken(FHlslToken(Identifier, LiteralType), Tokenizer);
 				}
 				else if (Tokenizer.MatchIdentifier(Identifier))
 				{
 					if (!FCString::Strcmp(*Identifier, TEXT("true")) || !FCString::Strcmp(*Identifier, TEXT("false")))
 					{
-						AddToken(FHlslToken(*Identifier, ELiteralType::Bool), Tokenizer);
+						EmplaceToken(FHlslToken(*Identifier, ELiteralType::Bool), Tokenizer);
 					}
 					else if (MatchSymbolToken(*Identifier, nullptr, SymbolToken, nullptr, true))
 					{
-						AddToken(FHlslToken(SymbolToken, Identifier), Tokenizer);
+						EmplaceToken(FHlslToken(SymbolToken, Identifier), Tokenizer);
 					}
 					else
 					{
-						AddToken(FHlslToken(Identifier), Tokenizer);
+						EmplaceToken(FHlslToken(Identifier), Tokenizer);
 					}
 				}
 				else if (Tokenizer.MatchSymbol(SymbolToken, Identifier))
 				{
-					AddToken(FHlslToken(SymbolToken, Identifier), Tokenizer);
+					EmplaceToken(FHlslToken(SymbolToken, Identifier), Tokenizer);
 				}
 				else if (Tokenizer.MatchQuotedString(Identifier))
 				{
-					AddToken(FHlslToken(EHlslToken::StringConstant, Identifier), Tokenizer);
+					EmplaceToken(FHlslToken(EHlslToken::StringConstant, Identifier), Tokenizer);
 				}
 				else if (Tokenizer.HasCharsAvailable())
 				{
@@ -1316,7 +1319,7 @@ namespace CrossCompiler
 			{
 				Tokenizer.Line = Line - 1;
 				Tokenizer.SkipWhitespaceInLine();
-				FString Filename;
+				TStringBuilder<256> Filename;
 				if (Tokenizer.MatchQuotedString(Filename))
 				{
 					Tokenizer.Filename = Filename;
@@ -1324,14 +1327,20 @@ namespace CrossCompiler
 			}
 			else
 			{
-				FString LineString = TEXT("#line ") + Tokenizer.ReadToEndOfLine();
+				TStringBuilder<128> LineString;
+				LineString += TEXTVIEW("#line ");
+				Tokenizer.ReadToEndOfLine(LineString);
+
 				CompilerMessages.SourceError(*FString::Printf(TEXT("Malformed #line directive: %s!"), *LineString));
 			}
 		}
 		else if (Tokenizer.MatchString(MATCH_TARGET(TEXT("#pragma"))))
 		{
-			FString Pragma = TEXT("#pragma") + Tokenizer.ReadToEndOfLine(false);
-			Scanner.AddToken(FHlslToken(EHlslToken::Pragma, Pragma), Tokenizer);
+			TStringBuilder<128> Pragma;
+			Pragma += TEXTVIEW("#pragma");
+			Tokenizer.ReadToEndOfLine(Pragma, false);
+
+			Scanner.EmplaceToken(FHlslToken(EHlslToken::Pragma, Pragma), Tokenizer);
 		}
 		else if (Tokenizer.MatchString(MATCH_TARGET(TEXT("#if 0"))))
 		{
@@ -1354,13 +1363,16 @@ namespace CrossCompiler
 			}
 			else
 			{
-				FString Directive = TEXT("#if 0") + Tokenizer.ReadToEndOfLine();
+				TStringBuilder<128> Directive;
+				Directive += TEXT("#if 0");
+				Tokenizer.ReadToEndOfLine(Directive);
 				CompilerMessages.SourceWarning(*FString::Printf(TEXT("Unhandled preprocessor directive (%.500s); HlslParser requires preprocessed input!"), *Directive));
 			}
 		}
 		else
 		{
-			FString Directive = Tokenizer.ReadToEndOfLine();
+			TStringBuilder<128> Directive;
+			Tokenizer.ReadToEndOfLine(Directive);
 			CompilerMessages.SourceWarning(*FString::Printf(TEXT("Unhandled token (%.500s); HlslParser requires preprocessed input!"), *Directive));
 		}
 
diff --git a/Engine/Source/Developer/ShaderCompilerCommon/Private/HlslLexer.h b/Engine/Source/Developer/ShaderCompilerCommon/Private/HlslLexer.h
--- a/Engine/Source/Developer/ShaderCompilerCommon/Private/HlslLexer.h
+++ b/Engine/Source/Developer/ShaderCompilerCommon/Private/HlslLexer.h
@@ -344,12 +344,12 @@ namespace CrossCompiler
 	{
 		EHlslToken Token;
 		ELiteralType LiteralType = ELiteralType::Unknown;
-		FString String;
+		FString String; // yuriy-todo: allocate this from linear allocator or simply point to the input source code and use FStringView
 		FSourceInfo SourceInfo;
 
-		explicit FHlslToken(const FString& Identifier) : Token(EHlslToken::Identifier), String(Identifier) { DebugId(); }
-		explicit FHlslToken(EHlslToken InToken, const FString& Identifier) : Token(InToken), String(Identifier) { DebugId(); }
-		explicit FHlslToken(const FString& InLiteral, ELiteralType InLiteralType) :
+		explicit FHlslToken(FStringView Identifier) : Token(EHlslToken::Identifier), String(Identifier) { DebugId(); }
+		explicit FHlslToken(EHlslToken InToken, FStringView Identifier) : Token(InToken), String(Identifier) { DebugId(); }
+		explicit FHlslToken(FStringView InLiteral, ELiteralType InLiteralType) :
 			Token(EHlslToken::Literal),
 			LiteralType(InLiteralType),
 			String(InLiteral)
@@ -397,7 +397,7 @@ namespace CrossCompiler
 		FCompilerMessages& CompilerMessages;
 
 		TArray<FHlslToken> Tokens;
-		void AddToken(const FHlslToken& Token, const struct FTokenizer& Tokenizer);
+		void EmplaceToken(FHlslToken&& Token, const struct FTokenizer& Tokenizer);
 
 		uint32 CurrentToken;
 
diff --git a/Engine/Source/Developer/ShaderCompilerCommon/Private/HlslUtils.h b/Engine/Source/Developer/ShaderCompilerCommon/Private/HlslUtils.h
--- a/Engine/Source/Developer/ShaderCompilerCommon/Private/HlslUtils.h
+++ b/Engine/Source/Developer/ShaderCompilerCommon/Private/HlslUtils.h
@@ -112,9 +112,9 @@ namespace CrossCompiler
 			return Data;
 		}
 
-		inline TCHAR* Strdup(const FString& String)
+		inline TCHAR* Strdup(const FStringView String)
 		{
-			return Strdup(*String, String.Len());
+			return Strdup(String.GetData(), String.Len());
 		}
 
 		TArray<Memory::FPage*, TInlineAllocator<8> > Pages;
