[Backout] - CL41031109
Original CL Desc
-----------------------------------------------------------------
[Backout] - CL40793982
Original CL Desc
-----------------------------------------------------------------
Metal MSC - Fix large metallib size by only embedding debug information when needed with r.Shaders.Symbols=1
The debug information is already embedded automatically when needed with -Zi & -Qembed_debug in MetalCompileShaderMSC.cpp and nothing else is needed.

#tests Local test with and without symbols in editor and packaged game

diff --git a/Engine/Shaders/Public/Platform/Metal/MetalCommon.ush b/Engine/Shaders/Public/Platform/Metal/MetalCommon.ush
--- a/Engine/Shaders/Public/Platform/Metal/MetalCommon.ush
+++ b/Engine/Shaders/Public/Platform/Metal/MetalCommon.ush
@@ -7,7 +7,7 @@
 #pragma once
 
 // Update this GUID to invalidate shader recompilation for only Metal shaders
-#pragma message("UESHADERMETADATA_VERSION A68B2D92-5E85-40D6-89F0-FCB78A6EAC75")
+#pragma message("UESHADERMETADATA_VERSION 92F6351B-26F7-43BE-A6B7-DA4B8E990D78")
 
 #if VERTEXSHADER && HULLSHADER
 #ifndef TESSELLATIONSHADER
diff --git a/Engine/Source/Developer/ShaderCompilerCommon/Private/ShaderConductorContext.cpp b/Engine/Source/Developer/ShaderCompilerCommon/Private/ShaderConductorContext.cpp
--- a/Engine/Source/Developer/ShaderCompilerCommon/Private/ShaderConductorContext.cpp
+++ b/Engine/Source/Developer/ShaderCompilerCommon/Private/ShaderConductorContext.cpp
@@ -683,9 +683,6 @@ namespace CrossCompiler
 		
         ShaderConductor::Compiler::Options ScOptions;
         ConvertScOptions(*Intermediates, Options, ScOptions, false, bGenerateSpirv);
-        
-        // Force PDB generation (required to generate DXIL reflection).
-        ScOptions.enableDebugInfo = true;
 
         // Compile HLSL source code to DXIL
         bool bSucceeded = false;
