Add support for inlining functions inside nested statements. Previously only calls from the parent function body were supported.

diff --git a/Engine/Source/Developer/ShaderCompilerCommon/Private/HlslUtils.cpp b/Engine/Source/Developer/ShaderCompilerCommon/Private/HlslUtils.cpp
--- a/Engine/Source/Developer/ShaderCompilerCommon/Private/HlslUtils.cpp
+++ b/Engine/Source/Developer/ShaderCompilerCommon/Private/HlslUtils.cpp
@@ -2112,31 +2112,30 @@ namespace UE::ShaderCompilerCommon
 
 static void VisitCompoundStatementChildren(
 	CrossCompiler::AST::FCompoundStatement* InRoot,
-	TFunction<void(CrossCompiler::AST::FNode* Node)> Visitor)
+	TFunction<void(CrossCompiler::AST::FNode*& Node)> Visitor)
 {
 	using namespace CrossCompiler::AST;
 
 	TArray<FCompoundStatement*, TInlineAllocator<128>> Stack;
 	Stack.Push(InRoot);
 
-	auto VisitAndConditionallyPushNode = [&Stack, &Visitor](FNode* Node)
+	auto VisitAndConditionallyPushNode = [&Stack, &Visitor](FNode*& Node)
 	{
 		if (!Node)
 		{
 			return false;
 		}
 
-		Visitor(Node);
-
+		bool bResult = false;
 		if (FCompoundStatement* CompoundStatement = Node->AsCompoundStatement())
 		{
 			Stack.Push(CompoundStatement);
-			return true;
-		}
-		else
-		{
-			return false;
+			bResult = true;
 		}
+
+		Visitor(Node);
+
+		return bResult;
 	};
 
 	while (!Stack.IsEmpty())
@@ -2144,7 +2143,7 @@ static void VisitCompoundStatementChildren(
 		FCompoundStatement* Node = Stack.Last();
 		Stack.Pop();
 
-		for (FNode* ChildStatement : Node->Statements)
+		for (FNode*& ChildStatement : Node->Statements)
 		{
 			if (VisitAndConditionallyPushNode(ChildStatement))
 			{
@@ -2168,7 +2167,7 @@ static void VisitCompoundStatementChildren(
 				{
 					for (FCaseStatement* CaseStatement : SwitchStatement->Body->CaseList->Cases)
 					{
-						for (FNode* CaseStatementChildNode : CaseStatement->Statements)
+						for (FNode*& CaseStatementChildNode : CaseStatement->Statements)
 						{
 							VisitAndConditionallyPushNode(CaseStatementChildNode);
 						}
@@ -2224,7 +2223,7 @@ static void VisitFunctionIdentifiers(
 		}
 	}
 
-	VisitCompoundStatementChildren(FunctionDefinition->Body, [&ProcessDeclaratorList, &ProcessDeclaration](FNode* StatementNode)
+	VisitCompoundStatementChildren(FunctionDefinition->Body, [&ProcessDeclaratorList, &ProcessDeclaration](FNode*& StatementNode)
 	{
 		if (FIterationStatement* IterationStatement = StatementNode->AsIterationStatement())
 		{
@@ -2302,7 +2301,7 @@ bool InlineFunction(FString& InOutSourceCode, FString& InOutEntryPoint, const FS
 	// Check that the inlined function meets the current big limitation: no return statements.
 
 	bool bInlinedFunctionCompatible = true;
-	VisitCompoundStatementChildren(InlinedFunctionDefinition->Body, [&bInlinedFunctionCompatible, &OutErrors] (FNode* Node)
+	VisitCompoundStatementChildren(InlinedFunctionDefinition->Body, [&bInlinedFunctionCompatible, &OutErrors] (FNode*& Node)
 	{
 		if (FJumpStatement* JumpStatement = Node->AsJumpStatement())
 		{
@@ -2353,7 +2352,6 @@ bool InlineFunction(FString& InOutSourceCode, FString& InOutEntryPoint, const FS
 	const CrossCompiler::FSourceInfo SourceInfo; // todo: more accurate/better one?
 
 	FFunctionDefinition* OutputFunction = OutputNodes.AllocNode<FFunctionDefinition>(SourceInfo);
-	FCompoundStatement* OutputBody = OutputNodes.AllocNode<FCompoundStatement>(SourceInfo);
 
 	TStringBuilder<256> OptimizedEntryPoint;
 	OptimizedEntryPoint << EntryFunctionDefinition->Prototype->Identifier->ToStringView();
@@ -2361,16 +2359,15 @@ bool InlineFunction(FString& InOutSourceCode, FString& InOutEntryPoint, const FS
 
 	OutputFunction->Prototype = EntryFunctionDefinition->Prototype;
 	OutputFunction->Prototype->Identifier = FIdentifier::New(&OutputNodes.Allocator, OptimizedEntryPoint);
-	OutputFunction->Body = OutputBody;
+	OutputFunction->Body = EntryFunctionDefinition->Body; // rewrite the original nodes
 
 	// copy over original statements until we hit a call to the function we're inlining
 
-	// TODO: recursively process compound statements
-	for (FNode* Statement : EntryFunctionDefinition->Body->Statements)
+	VisitCompoundStatementChildren(OutputFunction->Body,
+		[&bInlinedFunctionCompatible, &AsCallToInlinedFunction, &OutputNodes, &SourceInfo, InlinedFunctionDefinition]
+		(FNode*& Node)
 	{
-		FNode* StatementToAdd = Statement; // shallow clone of the original statement by default
-
-		if (FExpressionStatement* CallExpressionStatement = AsCallToInlinedFunction(Statement))
+		if (FExpressionStatement* CallExpressionStatement = AsCallToInlinedFunction(Node))
 		{
 			FCompoundStatement* InlinedBody = OutputNodes.AllocNode<FCompoundStatement>(SourceInfo);
 
@@ -2445,20 +2442,13 @@ bool InlineFunction(FString& InOutSourceCode, FString& InOutEntryPoint, const FS
 				InlinedBody->Statements.Add(AssignmentStatement);
 			}
 
-			StatementToAdd = InlinedBody;
+			Node = InlinedBody;
 		}
-
-		OutputFunction->Body->Statements.Add(StatementToAdd);
-	}
-
-	OutputNodes.Nodes.Add(OutputFunction);
+	});
 
 	FASTWriter GeneratedCodeWriter(InOutSourceCode);
 
-	if (!OutputNodes.Nodes.IsEmpty())
-	{
-		OutputNodes.Nodes.Last()->Write(GeneratedCodeWriter);
-	}
+	OutputFunction->Write(GeneratedCodeWriter);
 
 	InOutEntryPoint = OptimizedEntryPoint;
 
