[Shaders] Allow enabling source information via CVar "r.ShaderCompiler.ForceSpirvDebugInfo" in all SPIR-V backends.

This enables the option to include full shader source information with OpLine instructions in the SPIR-V module to inspect shaders in profiling tools such as RenderDoc.
This was previoulsy only supported in one of the many SPIR-V shader backends. The CVar is read per target platform.

#rnx

diff --git a/Engine/Source/Developer/Apple/MetalShaderFormat/Private/MetalCompileShaderMSC.cpp b/Engine/Source/Developer/Apple/MetalShaderFormat/Private/MetalCompileShaderMSC.cpp
--- a/Engine/Source/Developer/Apple/MetalShaderFormat/Private/MetalCompileShaderMSC.cpp
+++ b/Engine/Source/Developer/Apple/MetalShaderFormat/Private/MetalCompileShaderMSC.cpp
@@ -597,18 +597,12 @@ void FMetalCompileShaderMSC::DoCompileMetalShader(
 		CrossCompiler::FShaderConductorContext CompilerContext;
 
 		// Initialize compilation options for ShaderConductor
-		CrossCompiler::FShaderConductorOptions Options;
+		CrossCompiler::FShaderConductorOptions Options(Input.Environment.CompilerFlags);
 
 		Options.TargetEnvironment = CrossCompiler::FShaderConductorOptions::ETargetEnvironment::Vulkan_1_2;
 
-		// Enable HLSL 2021 if specified
-		if (Input.Environment.CompilerFlags.Contains(CFLAG_HLSL2021))
-		{
-			Options.HlslVersion = 2021;
-		}
-		
 		Options.bEnable16bitTypes = true;
-		
+
 		FMetalShaderParameterParserPlatformConfiguration PlatformConfiguration;
 		FShaderParameterParser ShaderParameterParser(PlatformConfiguration);
 		if (!ShaderParameterParser.ParseAndModify(Input, Output.Errors, PreprocessedShader))
diff --git a/Engine/Source/Developer/Apple/MetalShaderFormat/Private/MetalCompileShaderSPIRV.cpp b/Engine/Source/Developer/Apple/MetalShaderFormat/Private/MetalCompileShaderSPIRV.cpp
--- a/Engine/Source/Developer/Apple/MetalShaderFormat/Private/MetalCompileShaderSPIRV.cpp
+++ b/Engine/Source/Developer/Apple/MetalShaderFormat/Private/MetalCompileShaderSPIRV.cpp
@@ -227,16 +227,9 @@ void FMetalCompileShaderSPIRV::DoCompileMetalShader(
 		CrossCompiler::FShaderConductorContext CompilerContext;
 
 		// Initialize compilation options for ShaderConductor
-		CrossCompiler::FShaderConductorOptions Options;
+		CrossCompiler::FShaderConductorOptions Options(Input.Environment.CompilerFlags);
 
 		Options.TargetEnvironment = CrossCompiler::FShaderConductorOptions::ETargetEnvironment::Vulkan_1_1;
-		Options.bWarningsAsErrors = Input.Environment.CompilerFlags.Contains(CFLAG_WarningsAsErrors);
-
-		// Enable HLSL 2021 if specified
-		if (Input.Environment.CompilerFlags.Contains(CFLAG_HLSL2021))
-		{
-			Options.HlslVersion = 2021;
-		}
 
 		// Always disable FMA pass for Pixel and Compute shader,
 		// otherwise determine whether [[position, invariant]] qualifier is available in Metal or not.
diff --git a/Engine/Source/Developer/ShaderCompilerCommon/Private/ShaderConductorContext.cpp b/Engine/Source/Developer/ShaderCompilerCommon/Private/ShaderConductorContext.cpp
--- a/Engine/Source/Developer/ShaderCompilerCommon/Private/ShaderConductorContext.cpp
+++ b/Engine/Source/Developer/ShaderCompilerCommon/Private/ShaderConductorContext.cpp
@@ -17,6 +17,21 @@ namespace CrossCompiler
 
 #if PLATFORM_MAC || PLATFORM_WINDOWS || PLATFORM_LINUX
 
+	void FShaderConductorOptions::ApplyCommonCompilerFlags(const FShaderCompilerFlags& InCompilerFlags)
+	{
+		// Treat all warnings as errors
+		bWarningsAsErrors = InCompilerFlags.Contains(CFLAG_WarningsAsErrors);
+
+		// Include shader source code in the shader debug info. This allows source code correlation in shader profiling tools.
+		bEnableDebugInfo = InCompilerFlags.Contains(CFLAG_ForceSpirvDebugInfo);
+
+		// Enable HLSL 2021 if specified
+		if (InCompilerFlags.Contains(CFLAG_HLSL2021))
+		{
+			HlslVersion = 2021;
+		}
+	}
+
 	static void ScRewriteWrapper(
 		const ShaderConductor::Compiler::SourceDesc& InSourceDesc,
 		const ShaderConductor::Compiler::Options& InOptions,
diff --git a/Engine/Source/Developer/ShaderCompilerCommon/Public/ShaderConductorContext.h b/Engine/Source/Developer/ShaderCompilerCommon/Public/ShaderConductorContext.h
--- a/Engine/Source/Developer/ShaderCompilerCommon/Public/ShaderConductorContext.h
+++ b/Engine/Source/Developer/ShaderCompilerCommon/Public/ShaderConductorContext.h
@@ -64,6 +64,14 @@ namespace CrossCompiler
 	/** Wrapper structure to pass options descriptor to ShaderConductor. This is mapped to <struct ShaderConductor::Compiler::Options>. */
 	struct FShaderConductorOptions
 	{
+		FShaderConductorOptions() = default;
+
+		/** Initializes the options with a call to ApplyCommonCompilerFlags(). */
+		FORCEINLINE FShaderConductorOptions(const FShaderCompilerFlags& InCompilerFlags)
+		{
+			ApplyCommonCompilerFlags(InCompilerFlags);
+		}
+
 		/** Removes unused global variables and resources. This can only be used in the HLSL rewrite pass, i.e. 'RewriteHlslSource'. */
 		bool bRemoveUnusedGlobals = false;
 
@@ -121,6 +129,12 @@ namespace CrossCompiler
 		 * Use "preset(relax-nested-expr)" for pre-defined set of optimization passes to relax nested expressions.
 		 */
 		FString SpirvCustomOptimizationPasses;
+
+		/**
+		 * Applies common compiler flags to this options struct: WarningsAsErrors, ForceSpirvDebugInfo, and HLSL2021.
+		 * Other flags must be individually applied depending on the shader backend.
+		 */
+		UE_API void ApplyCommonCompilerFlags(const FShaderCompilerFlags& InCompilerFlags);
 	};
 
 	/** Target high level languages for ShaderConductor output. */
diff --git a/Engine/Source/Developer/ShaderFormatOpenGL/Private/OpenGLShaderCompiler.cpp b/Engine/Source/Developer/ShaderFormatOpenGL/Private/OpenGLShaderCompiler.cpp
--- a/Engine/Source/Developer/ShaderFormatOpenGL/Private/OpenGLShaderCompiler.cpp
+++ b/Engine/Source/Developer/ShaderFormatOpenGL/Private/OpenGLShaderCompiler.cpp
@@ -2808,17 +2808,10 @@ static bool CompileToGlslWithShaderConductor(
 	const bool bDumpDebugInfo = Input.DumpDebugInfoEnabled();
 
 	// Initialize compilation options for ShaderConductor
-	CrossCompiler::FShaderConductorOptions Options;
+	CrossCompiler::FShaderConductorOptions Options(Input.Environment.CompilerFlags);
 	Options.bDisableScalarBlockLayout = true;
 	Options.bForceStorageImageFormat = true;
-	Options.bWarningsAsErrors = Input.Environment.CompilerFlags.Contains(CFLAG_WarningsAsErrors);
 	
-	// Enable HLSL 2021 if specified
-	if (Input.Environment.CompilerFlags.Contains(CFLAG_HLSL2021))
-	{
-		Options.HlslVersion = 2021;
-	}
-
 	// Convert input strings from FString to ANSI strings
 	std::string SourceData(TCHAR_TO_UTF8(*PreprocessedShader));
 	std::string FileName(TCHAR_TO_UTF8(*Input.VirtualSourceFilePath));
diff --git a/Engine/Source/Developer/VulkanShaderFormat/Private/SpirVShaderCompiler.inl b/Engine/Source/Developer/VulkanShaderFormat/Private/SpirVShaderCompiler.inl
--- a/Engine/Source/Developer/VulkanShaderFormat/Private/SpirVShaderCompiler.inl
+++ b/Engine/Source/Developer/VulkanShaderFormat/Private/SpirVShaderCompiler.inl
@@ -1327,9 +1327,8 @@ static bool CompileWithShaderConductor(
 	CompilerContext.LoadSource(PreprocessedShader, Input.VirtualSourceFilePath, InternalState.GetEntryPointName(), InternalState.GetShaderFrequency(), &AdditionalDefines, &ExtraDxcArgs);
 
 	// Initialize compilation options for ShaderConductor
-	CrossCompiler::FShaderConductorOptions Options;
+	CrossCompiler::FShaderConductorOptions Options(Input.Environment.CompilerFlags);
 	Options.TargetEnvironment = InternalState.GetMinimumTargetEnvironment();
-	Options.bWarningsAsErrors = Input.Environment.CompilerFlags.Contains(CFLAG_WarningsAsErrors);
 
 	// VK_EXT_scalar_block_layout is required by raytracing and by Nanite (so expect it to be present in SM6/Vulkan_1_3)
 	Options.bDisableScalarBlockLayout = !(InternalState.IsRayTracingShader() || InternalState.IsSM6());
@@ -1346,12 +1345,6 @@ static bool CompileWithShaderConductor(
 		Options.bEnable16bitTypes = true;
 	}
 
-	// Enable HLSL 2021 if specified
-	if (Input.Environment.CompilerFlags.Contains(CFLAG_HLSL2021))
-	{
-		Options.HlslVersion = 2021;
-	}
-
 	if (InternalState.bDebugDump)
 	{
 		SpirvCreateDXCCompileBatchFiles(CompilerContext, InternalState, Options);
