[shaders] fix shaderconductor error/warning parsing to better identify where error messages start and end.

diff --git a/Engine/Source/Developer/ShaderCompilerCommon/Private/ShaderConductorContext.cpp b/Engine/Source/Developer/ShaderCompilerCommon/Private/ShaderConductorContext.cpp
--- a/Engine/Source/Developer/ShaderCompilerCommon/Private/ShaderConductorContext.cpp
+++ b/Engine/Source/Developer/ShaderCompilerCommon/Private/ShaderConductorContext.cpp
@@ -559,17 +559,10 @@ namespace CrossCompiler
 	// Converts the error blob from ShaderConductor into an array of error reports (of type FShaderCompilerError).
 	static void ConvertScCompileErrors(ShaderConductor::Blob& ErrorBlob, TArray<FShaderCompilerError>& OutErrors)
 	{
-		// Convert blob into FString
-		FString ErrorString;
-		if (ConvertScBlobToFString(&ErrorBlob, ErrorString))
-		{
-			// Convert FString into array of FString (one for each line)
-			TArray<FString> ErrorStringLines;
-			ErrorString.ParseIntoArray(ErrorStringLines, TEXT("\n"));
+		FAnsiStringView ErrorString(static_cast<const ANSICHAR*>(ErrorBlob.Data()), ErrorBlob.Size());
 
-			// Forward parsed array of lines to primary conversion function
-			FShaderConductorContext::ConvertCompileErrors(MoveTemp(ErrorStringLines), OutErrors);
-		}
+		// Forward parsed array of lines to primary conversion function
+		FShaderConductorContext::ConvertCompileErrors(ErrorString, OutErrors);
 	}
 
 	FShaderConductorTarget::FShaderConductorTarget()
@@ -922,33 +915,57 @@ namespace CrossCompiler
 		return CmdLineArgs;
 	}
 
-	void FShaderConductorContext::ConvertCompileErrors(TArray<FString>&& ErrorStringLines, TArray<FShaderCompilerError>& OutErrors)
+	void FShaderConductorContext::ConvertCompileErrors(FAnsiStringView ErrorsString, TArray<FShaderCompilerError>& OutErrors)
 	{
-		// Returns whether the specified line in the 'ErrorStringLines' array has a line marker.
-		auto HasErrorLineMarker = [&ErrorStringLines](int32 LineIndex)
-		{
-			if (LineIndex + 2 < ErrorStringLines.Num())
+		auto ScanPastDigits = [](const ANSICHAR*& Ptr)
 			{
-				return IsTextLineDxcLineMarker(ErrorStringLines[LineIndex + 2]);
-			}
-			return false;
-		};
+				while (FChar::IsDigit(*(++Ptr))) {}
+			};
 
-		// Iterate over all errors. Most (but not all) contain a highlighted line and line marker.
-		for (int32 LineIndex = 0; LineIndex < ErrorStringLines.Num();)
+		// if the null terminator character is included in the string view, remove it before parsing so whitespace trimming works correctly
+		if (ErrorsString.EndsWith(0))
 		{
-			if (HasErrorLineMarker(LineIndex))
+			ErrorsString.LeftChopInline(1);
+		}
+
+		const ANSICHAR* ErrorBlobAnsiStart = ErrorsString.GetData();
+		const ANSICHAR* ErrorBlobAnsiCur = ErrorBlobAnsiStart;
+		const ANSICHAR* ErrorBlobAnsiEnd = ErrorBlobAnsiCur + ErrorsString.Len();
+		const ANSICHAR* LastNewLine = nullptr;
+		const ANSICHAR* CurErrorStart = nullptr;
+		int32 ColonCount = 0;
+
+		// since DXC doesn't give us a good delimiter for error messages, we look for a line containing 3 colons with digits in between (which will be the case
+		// for a file/line marker, i.e. /My/File/Shader.ush:269:29:) and use the start of that line as the start of a diagnostic message (and the end of the previous)
+		while (ErrorBlobAnsiCur < ErrorBlobAnsiEnd)
+		{
+			if (*ErrorBlobAnsiCur == '\n')
 			{
-				// Add current line as error with highlighted source line (LineIndex+1) and line marker (LineIndex+2)
-				OutErrors.Emplace(MoveTemp(ErrorStringLines[LineIndex]), MoveTemp(ErrorStringLines[LineIndex + 1]), MoveTemp(ErrorStringLines[LineIndex + 2]));
-				LineIndex += 3;
+				LastNewLine = ErrorBlobAnsiCur;
 			}
-			else
+			else if (*ErrorBlobAnsiCur == ':')
 			{
-				// Add current line as single error
-				OutErrors.Emplace(MoveTemp(ErrorStringLines[LineIndex]));
-				LineIndex += 1;
+				ScanPastDigits(ErrorBlobAnsiCur);
+				if (*ErrorBlobAnsiCur++ == ':')
+				{
+					ScanPastDigits(ErrorBlobAnsiCur);
+					if (*ErrorBlobAnsiCur++ == ':')
+					{
+						if (CurErrorStart != nullptr)
+						{
+							OutErrors.Add(FString(FAnsiStringView(CurErrorStart, LastNewLine - CurErrorStart).TrimStartAndEnd()));
+						}
+						CurErrorStart = LastNewLine != nullptr ? (LastNewLine + 1) : ErrorBlobAnsiStart;
+					}
+				}
 			}
+			ErrorBlobAnsiCur++;
+		}
+		// after the loop above ErrorBlobAnsiCur should be pointing to the end of the error blob, and CurErrorStart should point to the
+		// start of the last error, so we need to add one more error containing the characters between the two
+		if (CurErrorStart != nullptr)
+		{
+			OutErrors.Add(FString(FAnsiStringView(CurErrorStart, ErrorBlobAnsiEnd - CurErrorStart).TrimStartAndEnd()));
 		}
 	}
 
diff --git a/Engine/Source/Developer/ShaderCompilerCommon/Public/ShaderConductorContext.h b/Engine/Source/Developer/ShaderCompilerCommon/Public/ShaderConductorContext.h
--- a/Engine/Source/Developer/ShaderCompilerCommon/Public/ShaderConductorContext.h
+++ b/Engine/Source/Developer/ShaderCompilerCommon/Public/ShaderConductorContext.h
@@ -257,8 +257,8 @@ namespace CrossCompiler
 		}
 
 	public:
-		/** Convert array of error string lines into array of <FShaderCompilerError>. */
-		static UE_API void ConvertCompileErrors(TArray<FString>&& ErrorStringLines, TArray<FShaderCompilerError>& OutErrors);
+		/** Convert raw ANSI string buffer containing DXC errors into array of <FShaderCompilerError>. */
+		static UE_API void ConvertCompileErrors(FAnsiStringView ErrorsString, TArray<FShaderCompilerError>& OutErrors);
 
 		/** Returns whether the specified variable name denotes an intermediate output variable.
 		This is only true for a special identifiers generated by DXC to communicate patch constant data in the Hull Shader. */
diff --git a/Engine/Source/Developer/Windows/ShaderFormatD3D/Private/D3DShaderCompiler.cpp b/Engine/Source/Developer/Windows/ShaderFormatD3D/Private/D3DShaderCompiler.cpp
--- a/Engine/Source/Developer/Windows/ShaderFormatD3D/Private/D3DShaderCompiler.cpp
+++ b/Engine/Source/Developer/Windows/ShaderFormatD3D/Private/D3DShaderCompiler.cpp
@@ -1173,15 +1173,8 @@ static bool CompileAndProcessD3DShaderFXCExt(
 			Output.bSucceeded = false;
 		}
 		
-		TArray<FString> FilteredErrorsFromResourceCounts;
-		if (!ValidateResourceCounts(CompileData, FilteredErrorsFromResourceCounts))
+		if (!ValidateResourceCounts(CompileData, Output.Errors))
 		{
-			for (FString& error : FilteredErrorsFromResourceCounts)
-			{
-				FShaderCompilerError NewError;
-				NewError.StrippedErrorMessage = error;
-				Output.Errors.Add(NewError);
-			}
 			Result = E_FAIL;
 			Output.bSucceeded = false;
 		}
diff --git a/Engine/Source/Developer/Windows/ShaderFormatD3D/Private/D3DShaderCompiler.inl b/Engine/Source/Developer/Windows/ShaderFormatD3D/Private/D3DShaderCompiler.inl
--- a/Engine/Source/Developer/Windows/ShaderFormatD3D/Private/D3DShaderCompiler.inl
+++ b/Engine/Source/Developer/Windows/ShaderFormatD3D/Private/D3DShaderCompiler.inl
@@ -368,7 +368,7 @@ template <typename ID3D1xShaderReflection, typename D3D1x_SHADER_DESC, typename
 // Currently limited for hard-coded root signature setup (see: FD3D12Adapter::StaticGraphicsRootSignature)
 // In theory this limitation is only required for DX12, but we don't want a shader to compile on DX11 while not working on DX12.
 // (DX11 has an API limit on 128 SRVs, 16 Samplers, 8 UAVs and 14 CBs but if you go over these values then the shader won't compile)
-inline bool ValidateResourceCounts(const FD3DShaderCompileData& CompileData, TArray<FString>& OutFilteredErrors)
+inline bool ValidateResourceCounts(const FD3DShaderCompileData& CompileData, TArray<FShaderCompilerError>& OutErrors)
 {
 	const bool bTooManySRVs     = !CompileData.bBindlessEnabled && CompileData.NumSRVs     > CompileData.MaxSRVs;
 	const bool bTooManyUAVs     = !CompileData.bBindlessEnabled && CompileData.NumUAVs     > CompileData.MaxUAVs;
@@ -379,22 +379,22 @@ inline bool ValidateResourceCounts(const FD3DShaderCompileData& CompileData, TAr
 	{
 		if (bTooManySRVs)
 		{
-			OutFilteredErrors.Add(FString::Printf(TEXT("Shader is using too many SRVs: %d (only %d supported)"), CompileData.NumSRVs, CompileData.MaxSRVs));
+			OutErrors.Add(FString::Printf(TEXT("Shader is using too many SRVs: %d (only %d supported)"), CompileData.NumSRVs, CompileData.MaxSRVs));
 		}
 
 		if (bTooManySamplers)
 		{
-			OutFilteredErrors.Add(FString::Printf(TEXT("Shader is using too many Samplers: %d (only %d supported)"), CompileData.NumSamplers, CompileData.MaxSamplers));
+			OutErrors.Add(FString::Printf(TEXT("Shader is using too many Samplers: %d (only %d supported)"), CompileData.NumSamplers, CompileData.MaxSamplers));
 		}
 
 		if (bTooManyUAVs)
 		{
-			OutFilteredErrors.Add(FString::Printf(TEXT("Shader is using too many UAVs: %d (only %d supported)"), CompileData.NumUAVs, CompileData.MaxUAVs));
+			OutErrors.Add(FString::Printf(TEXT("Shader is using too many UAVs: %d (only %d supported)"), CompileData.NumUAVs, CompileData.MaxUAVs));
 		}
 
 		if (bTooManyCBs)
 		{
-			OutFilteredErrors.Add(FString::Printf(TEXT("Shader is using too many Constant Buffers: %d (only %d supported)"), CompileData.NumCBs, CompileData.MaxCBs));
+			OutErrors.Add(FString::Printf(TEXT("Shader is using too many Constant Buffers: %d (only %d supported)"), CompileData.NumCBs, CompileData.MaxCBs));
 		}
 
 		return false;
diff --git a/Engine/Source/Developer/Windows/ShaderFormatD3D/Private/D3DShaderCompilerDXC.cpp b/Engine/Source/Developer/Windows/ShaderFormatD3D/Private/D3DShaderCompilerDXC.cpp
--- a/Engine/Source/Developer/Windows/ShaderFormatD3D/Private/D3DShaderCompilerDXC.cpp
+++ b/Engine/Source/Developer/Windows/ShaderFormatD3D/Private/D3DShaderCompilerDXC.cpp
@@ -48,11 +48,6 @@ THIRD_PARTY_INCLUDES_END
 #include "DXCUtils.inl"
 #include "D3DShaderCompiler.inl"
 
-FORCENOINLINE static void DXCFilterShaderCompileWarnings(const FString& CompileWarnings, TArray<FString>& FilteredWarnings)
-{
-	CompileWarnings.ParseIntoArray(FilteredWarnings, TEXT("\n"), true);
-}
-
 static bool IsGlobalConstantBufferSupported(const FShaderTarget& Target)
 {
 	switch (Target.Frequency)
@@ -843,11 +838,10 @@ bool CompileAndProcessD3DShaderDXC(
 		Ar << DebugData;
 	}
 
-	TArray<FString> FilteredErrors;
 	if (DxcErrorBlob && DxcErrorBlob->GetBufferSize())
 	{
-		FString ErrorString = DxcBlobEncodingToFString(DxcErrorBlob);
-		DXCFilterShaderCompileWarnings(ErrorString, FilteredErrors);
+		FAnsiStringView ErrorsString(static_cast<ANSICHAR*>(DxcErrorBlob->GetBufferPointer()), DxcErrorBlob->GetBufferSize());
+		CrossCompiler::FShaderConductorContext::ConvertCompileErrors(ErrorsString, Output.Errors);
 	}
 
 	if (SUCCEEDED(D3DCompileToDxilResult))
@@ -1103,7 +1097,7 @@ bool CompileAndProcessD3DShaderDXC(
 						}
 					}
 
-					FilteredErrors.Add(ErrorString);
+					Output.Errors.Add(MoveTemp(ErrorString));
 					Output.bSucceeded = false;
 				}
 			}
@@ -1138,7 +1132,7 @@ bool CompileAndProcessD3DShaderDXC(
 				);
 		}
 
-		if (!ValidateResourceCounts(CompileData, FilteredErrors))
+		if (!ValidateResourceCounts(CompileData, Output.Errors))
 		{
 			Output.bSucceeded = false;
 		}
@@ -1288,20 +1282,17 @@ bool CompileAndProcessD3DShaderDXC(
 	else
 	{
 		// If we failed and didn't get any error messages back from the compile call try and get a system error message.
-		if (FilteredErrors.Num() == 0)
+		if (Output.Errors.Num() == 0)
 		{
 			TCHAR ErrorMsg[1024];
 			FPlatformMisc::GetSystemErrorMessage(ErrorMsg, UE_ARRAY_COUNT(ErrorMsg), (int)D3DCompileToDxilResult);
 			const bool bKnownError = ErrorMsg[0] != TEXT('\0');
 
 			FString ErrorString = FString::Printf(TEXT("D3DCompileToDxil failed. Error code: %s (0x%08X)."), bKnownError ? ErrorMsg : TEXT("Unknown error"), (int)D3DCompileToDxilResult);
-
-			FilteredErrors.Add(ErrorString);
+			Output.Errors.Add(MoveTemp(ErrorString));
 		}
 	}
 
-	// Move intermediate filtered errors into compiler context for unification.
-	CrossCompiler::FShaderConductorContext::ConvertCompileErrors(MoveTemp(FilteredErrors), Output.Errors);
 
 	return Output.bSucceeded;
 }
diff --git a/Engine/Source/Developer/Windows/ShaderFormatD3D/Private/ShaderFormatD3D.h b/Engine/Source/Developer/Windows/ShaderFormatD3D/Private/ShaderFormatD3D.h
--- a/Engine/Source/Developer/Windows/ShaderFormatD3D/Private/ShaderFormatD3D.h
+++ b/Engine/Source/Developer/Windows/ShaderFormatD3D/Private/ShaderFormatD3D.h
@@ -83,7 +83,7 @@ bool CompileAndProcessD3DShaderDXC(
 	FShaderCompilerOutput& Output);
 
 struct FD3DShaderCompileData;
-bool ValidateResourceCounts(const FD3DShaderCompileData& CompiledData, TArray<FString>& OutFilteredErrors);
+bool ValidateResourceCounts(const FD3DShaderCompileData& CompiledData, TArray<FShaderCompilerError>& OutFilteredErrors);
 
 template <typename TBlobType>
 TConstArrayView<uint8> MakeArrayViewFromBlob(const TRefCountPtr<TBlobType> Blob)
