[shaders] rework to satisfy static analysis; the error it called out was not possible in practice, but this is slightly clearer anyways

diff --git a/Engine/Source/Developer/ShaderCompilerCommon/Private/ShaderConductorContext.cpp b/Engine/Source/Developer/ShaderCompilerCommon/Private/ShaderConductorContext.cpp
--- a/Engine/Source/Developer/ShaderCompilerCommon/Private/ShaderConductorContext.cpp
+++ b/Engine/Source/Developer/ShaderCompilerCommon/Private/ShaderConductorContext.cpp
@@ -931,7 +931,7 @@ namespace CrossCompiler
 		const ANSICHAR* ErrorBlobAnsiStart = ErrorsString.GetData();
 		const ANSICHAR* ErrorBlobAnsiCur = ErrorBlobAnsiStart;
 		const ANSICHAR* ErrorBlobAnsiEnd = ErrorBlobAnsiCur + ErrorsString.Len();
-		const ANSICHAR* LastNewLine = nullptr;
+		const ANSICHAR* LastLineStart = ErrorBlobAnsiStart;
 		const ANSICHAR* CurErrorStart = nullptr;
 		int32 ColonCount = 0;
 
@@ -941,7 +941,7 @@ namespace CrossCompiler
 		{
 			if (*ErrorBlobAnsiCur == '\n')
 			{
-				LastNewLine = ErrorBlobAnsiCur;
+				LastLineStart = ErrorBlobAnsiCur + 1;
 			}
 			else if (*ErrorBlobAnsiCur == ':')
 			{
@@ -953,9 +953,9 @@ namespace CrossCompiler
 					{
 						if (CurErrorStart != nullptr)
 						{
-							OutErrors.Add(FString(FAnsiStringView(CurErrorStart, LastNewLine - CurErrorStart).TrimStartAndEnd()));
+							OutErrors.Add(FString(FAnsiStringView(CurErrorStart, LastLineStart - CurErrorStart).TrimStartAndEnd()));
 						}
-						CurErrorStart = LastNewLine != nullptr ? (LastNewLine + 1) : ErrorBlobAnsiStart;
+						CurErrorStart = LastLineStart;
 					}
 				}
 			}
