[SPIRV] Fix crash when iterating through instructions of an empty SPIR-V module.

The constructor of TSpirvForwardIterator needs to ignore null pointers, so that begin()/end() loops can iterate over empty SPIR-V modules with zero iterations successfully.

#rnx

diff --git a/Engine/Source/Developer/ShaderCompilerCommon/Public/SpirvCommon.h b/Engine/Source/Developer/ShaderCompilerCommon/Public/SpirvCommon.h
--- a/Engine/Source/Developer/ShaderCompilerCommon/Public/SpirvCommon.h
+++ b/Engine/Source/Developer/ShaderCompilerCommon/Public/SpirvCommon.h
@@ -30,7 +30,8 @@ public:
 	TSpirvForwardIterator(T* InPtr, bool bPointsAtHeader = false) :
 		Ptr(InPtr)
 	{
-		if (bPointsAtHeader)
+		// Ignore null pointers so that empty SPIR-V modules can still use begin()/end() loops with zero iterations.
+		if (bPointsAtHeader && Ptr != nullptr)
 		{
 			// Validate header of SPIR-V module
 			const uint32 MagicNumberValue = Ptr[0];
