Add support for customizing the file name for ShaderSymbols and ShaderTypeStats output files

#tests local cooks under varying dlc scenarios

diff --git a/Engine/Config/ConsoleVariables.ini b/Engine/Config/ConsoleVariables.ini
--- a/Engine/Config/ConsoleVariables.ini
+++ b/Engine/Config/ConsoleVariables.ini
@@ -93,6 +93,11 @@ r.DumpShaderDebugInfo=2
 ; A few tokens can be used for engine paths. {ProjectDir} is the project directory, {ProjectSavedDir} is the project saved dir and {Platform} is the shader platform.
 ; The quotes are REQUIRED for the {} tokens to work.
 ;r.Shaders.SymbolPathOverride="{ProjectSavedDir}\ShaderSymbols\{Platform}"
+; Uncomment and change to set a custom name for ShaderSymbols output files (.info, .zip, etc.) or ShaderTypeStats build metadata
+; A few tokens can be used for engine paths. {Platform} is the shader platform, and {DLC} is DLC name if set, {?DLC-} a dash if there is DLC, for optional formatting
+; The quotes are REQUIRED for the {} tokens to work.
+;r.Shaders.SymbolFileNameOverride="ShaderSymbols{?DLC-}{DLC}"
+;r.Shaders.ShaderTypeStatsFileNameOverride="ShaderTypeStats-{Platform}{?DLC-}{DLC}"
 
 ; If Linux editor crashes in FMallocBinned with callstack that mentions MeshUtilities, you may need to uncomment this.
 ;r.TriangleOrderOptimization=2
diff --git a/Engine/Source/Developer/ShaderCompilerCommon/Private/ShaderSymbolExport.cpp b/Engine/Source/Developer/ShaderCompilerCommon/Private/ShaderSymbolExport.cpp
--- a/Engine/Source/Developer/ShaderCompilerCommon/Private/ShaderSymbolExport.cpp
+++ b/Engine/Source/Developer/ShaderCompilerCommon/Private/ShaderSymbolExport.cpp
@@ -17,7 +17,6 @@
 
 DECLARE_LOG_CATEGORY_CLASS(LogShaderSymbolExport, Display, Display);
 
-static const TCHAR* ZipFileBaseLeafName = TEXT("ShaderSymbols");
 static const TCHAR* ZipFileExtension = TEXT(".zip");
 static const TCHAR* InfoFileExtension = TEXT(".info");
 
@@ -28,31 +27,36 @@ FShaderSymbolExport::FShaderSymbolExport(FName InShaderFormat)
 
 FShaderSymbolExport::~FShaderSymbolExport() = default;
 
-static void DeleteExisting(IPlatformFile& PlatformFile, const FString& Directory, const TCHAR* BaseLeafName, const TCHAR* Extension)
+static void DeleteExisting(IPlatformFile& PlatformFile, const FString& Directory, const FString& BaseLeafName, const TCHAR* Extension)
 {
 	TArray<FString> ExistingZips;
 	PlatformFile.FindFiles(ExistingZips, *Directory, Extension);
 
+	// To minimize over-deleting match the full name [BaseLeafName.]ext or the workers [BaseLeafName_]N.ext
+	FString BaseLeafRoot = BaseLeafName + TEXT(".");
+	FString BaseLeafWorker = BaseLeafName + TEXT("_");
+
 	for (const FString& ZipFile : ExistingZips)
 	{
-		if (FPathViews::GetPathLeaf(ZipFile).StartsWith(BaseLeafName))
+		const FStringView Leaf = FPathViews::GetPathLeaf(ZipFile);
+		if (Leaf.StartsWith(BaseLeafRoot) || Leaf.StartsWith(BaseLeafWorker))
 		{
 			PlatformFile.DeleteFile(*ZipFile);
 		}
 	}
 }
 
-static FString CreateNameAndDeleteOld(uint32 MultiprocessId, IPlatformFile& PlatformFile, const FString& ExportPath, const TCHAR* BaseLeafName, const TCHAR* Extension)
+static FString CreateNameAndDeleteOld(uint32 MultiprocessId, IPlatformFile& PlatformFile, const FString& ExportPath, const FString& BaseLeafName, const TCHAR* Extension)
 {
 	FString Name;
 	if (MultiprocessId == 0)
 	{
 		DeleteExisting(PlatformFile, ExportPath, BaseLeafName, Extension);
-		Name = FString::Printf(TEXT("%s%s"), BaseLeafName, Extension);
+		Name = FString::Printf(TEXT("%s%s"), *BaseLeafName, Extension);
 	}
 	else
 	{
-		Name = FString::Printf(TEXT("%s_%d%s"), BaseLeafName, MultiprocessId, Extension);
+		Name = FString::Printf(TEXT("%s_%d%s"), *BaseLeafName, MultiprocessId, Extension);
 	}
 	return Name;
 }
@@ -81,6 +85,10 @@ void FShaderSymbolExport::Initialize()
 				*(FPaths::ProjectSavedDir() / TEXT("ShaderSymbols") / ShaderFormat.ToString()));
 		}
 
+		// Setup ExportFileName to the default or the cvar override
+		ExportFileName = TEXT("ShaderSymbols");
+		GetShaderFileNameOverride(ExportFileName, TEXT("r.Shaders.SymbolFileNameOverride"), ShaderFormat);
+
 		IPlatformFile& PlatformFile = FPlatformFileManager::Get().GetPlatformFile();
 		bExportShaderSymbols = PlatformFile.CreateDirectoryTree(*ExportPath);
 
@@ -100,7 +108,7 @@ void FShaderSymbolExport::Initialize()
 
 			if (bSymbolsEnabled && (bExportAsZip || FParse::Param(FCommandLine::Get(), TEXT("ShaderSymbolsExportZip"))))
 			{
-				FString LeafName = CreateNameAndDeleteOld(MultiprocessId, PlatformFile, ExportPath, ZipFileBaseLeafName, ZipFileExtension);
+				FString LeafName = CreateNameAndDeleteOld(MultiprocessId, PlatformFile, ExportPath, ExportFileName, ZipFileExtension);
 				FString SingleFilePath = ExportPath / LeafName;
 
 				IFileHandle* OutputZipFile = PlatformFile.OpenWrite(*SingleFilePath);
@@ -121,7 +129,7 @@ void FShaderSymbolExport::Initialize()
 			if (bSymbolsInfoEnabled)
 			{
 				// if we are exporting collated shader pdb info into one file
-				FString LeafName = CreateNameAndDeleteOld(MultiprocessId, PlatformFile, ExportPath, ZipFileBaseLeafName, InfoFileExtension);
+				FString LeafName = CreateNameAndDeleteOld(MultiprocessId, PlatformFile, ExportPath, ExportFileName, InfoFileExtension);
 				InfoFilePath = ExportPath / LeafName;
 			}
 		}
@@ -247,6 +255,10 @@ void FShaderSymbolExport::WriteSymbolData(const FString& Filename, const FString
 void FShaderSymbolExport::NotifyShaderCompilersShutdown()
 {
 	FScopeLock Lock(&Cs);
+
+	// Used to match worker files, for example (if not overridden): [ShaderSymbols_]N.ext
+	const FString WorkerName = ExportFileName + TEXT("_");
+
 	if (ShaderInfos.Num())
 	{
 		if (InfoFilePath.Len())
@@ -254,11 +266,17 @@ void FShaderSymbolExport::NotifyShaderCompilersShutdown()
 			IFileManager& FileManager = IFileManager::Get();
 			if (bMultiprocessOwner)
 			{
-				// if we are the multiprocess owner merge in any other files we find
-				// we will chunk up the worker files into {Hash, Data} pairs, dedupe them with ours, and sort them all
+				// If we are the multiprocess owner merge in any other files we find which have the same FileName, outside of multiproc numbering
+				// We do not want to delete other unrelated files in the export path, potentially from other cook jobs
+				// We will chunk up the worker files into {Hash, Data} pairs, dedupe them with ours, and sort them all
 				IPlatformFile& PlatformFile = FPlatformFileManager::Get().GetPlatformFile();
 				TArray<FString> FilesToMergeIn;
 				PlatformFile.FindFiles(FilesToMergeIn, *ExportPath, InfoFileExtension);
+				FilesToMergeIn.RemoveAll([&WorkerName](const FString& FileName)
+					{
+						return !FPathViews::GetPathLeaf(FileName).StartsWith(WorkerName);
+					});
+
 				for (const FString& InfoFile : FilesToMergeIn)
 				{
 					TUniquePtr<FArchive> Reader = TUniquePtr<FArchive>(FileManager.CreateFileReader(*InfoFile));
@@ -337,9 +355,9 @@ void FShaderSymbolExport::NotifyShaderCompilersShutdown()
 		IPlatformFile& PlatformFile = FPlatformFileManager::Get().GetPlatformFile();
 		TArray<FString> ZipsToMergeIn;
 		PlatformFile.FindFiles(ZipsToMergeIn, *ExportPath, ZipFileExtension);
-		ZipsToMergeIn.RemoveAll([](const FString& FileName)
+		ZipsToMergeIn.RemoveAll([&WorkerName](const FString& FileName)
 			{
-				return FPathViews::GetBaseFilename(FileName) == ZipFileBaseLeafName;
+				return !FPathViews::GetPathLeaf(FileName).StartsWith(WorkerName);
 			});
 
 #if WITH_EDITOR // FZipArchiveReader is only available in editor
diff --git a/Engine/Source/Developer/ShaderCompilerCommon/Public/ShaderSymbolExport.h b/Engine/Source/Developer/ShaderCompilerCommon/Public/ShaderSymbolExport.h
--- a/Engine/Source/Developer/ShaderCompilerCommon/Public/ShaderSymbolExport.h
+++ b/Engine/Source/Developer/ShaderCompilerCommon/Public/ShaderSymbolExport.h
@@ -54,6 +54,7 @@ private:
 	TMap<FString, FSymbolFileInfo> ExportedSymbolInfo;
 	FString ExportPath;
 	FString InfoFilePath;
+	FString ExportFileName;
 	FCriticalSection Cs;
 	uint64 TotalSymbolDataBytes{ 0 };
 	uint64 TotalSymbolData{ 0 };
diff --git a/Engine/Source/Runtime/Engine/Private/ShaderCompiler/ShaderCompiler.cpp b/Engine/Source/Runtime/Engine/Private/ShaderCompiler/ShaderCompiler.cpp
--- a/Engine/Source/Runtime/Engine/Private/ShaderCompiler/ShaderCompiler.cpp
+++ b/Engine/Source/Runtime/Engine/Private/ShaderCompiler/ShaderCompiler.cpp
@@ -393,6 +393,28 @@ static TAutoConsoleVariable<FString> CVarShadersSymbolPathOverride(
 	TEXT("This setting can be overriden in any Engine.ini under the [ShaderCompiler] section."),
 	ECVF_ReadOnly);
 
+static TAutoConsoleVariable<FString> CVarShadersSymbolFileNameOverride(
+	TEXT("r.Shaders.SymbolFileNameOverride"),
+	"",
+	TEXT("Override base file name for shader symbol related aggregate outputs (.zip, .info).\n")
+	TEXT("'{Platform}' will be replaced with the shader platform string.\n")
+	TEXT("'{DLC}' will be replaced with the DLC name if there is one.\n")
+	TEXT("'{?DLC-}' will be replaced with a dash if there is DLC, for formatting.\n")
+	TEXT("Empty: use default 'ShaderSymbols'\n")
+	TEXT("This setting can be overridden in any Engine.ini under the [ShaderCompiler] section."),
+	ECVF_ReadOnly);
+
+static TAutoConsoleVariable<FString> CVarShaderTypeStatsFileNameOverride(
+	TEXT("r.Shaders.ShaderTypeStatsFileNameOverride"),
+	"",
+	TEXT("Override base file name for shader type outputs.\n")
+	TEXT("'{Platform}' will be replaced with the shader platform string.\n")
+	TEXT("'{DLC}' will be replaced with the DLC name if there is one.\n")
+	TEXT("'{?DLC-}' will be replaced with a dash if there is DLC, for formatting.\n")
+	TEXT("Empty: use default 'ShaderTypeStats-{Platform}'\n")
+	TEXT("This setting can be overridden in any Engine.ini under the [ShaderCompiler] section."),
+	ECVF_ReadOnly);
+
 static TAutoConsoleVariable<int32> CVarAllowUniqueDebugInfo(
 	TEXT("r.Shaders.AllowUniqueSymbols"),
 	0,
diff --git a/Engine/Source/Runtime/RenderCore/Private/ShaderCodeLibrary.cpp b/Engine/Source/Runtime/RenderCore/Private/ShaderCodeLibrary.cpp
--- a/Engine/Source/Runtime/RenderCore/Private/ShaderCodeLibrary.cpp
+++ b/Engine/Source/Runtime/RenderCore/Private/ShaderCodeLibrary.cpp
@@ -36,6 +36,7 @@ ShaderCodeLibrary.cpp: Bound shader state cache implementation.
 #include "RenderingThread.h"
 #include "Shader.h"
 #include "ShaderCodeArchive.h"
+#include "ShaderCompilerCore.h"
 #include "ShaderPipelineCache.h"
 #include "DataDrivenShaderPlatformInfo.h"
 #include "String/ParseTokens.h"
@@ -3822,7 +3823,8 @@ public:
 		{
 			if (PlatformCodeStats.NumShaderMaps > 0)
 			{
-				FString PlatformName = FGenericDataDrivenShaderPlatformInfo::GetName((EShaderPlatform)PlatformId).ToString();
+				const FName PlatformFName = FGenericDataDrivenShaderPlatformInfo::GetName((EShaderPlatform)PlatformId);
+				FString PlatformName = PlatformFName.ToString();
 				FString Filename = FPaths::Combine(*DebugInfoDir, *PlatformName, TEXT("ShaderTypeStats.csv"));
 				TMap<const FShaderType*, FShaderCodeStats::FPerTypeStats> CountersByTypeName;
 				for (const TPair<uint64, FShaderCodeStats::FPerTypeStats>& Pair : PlatformCodeStats.ShaderTypeStats)
@@ -3875,9 +3877,15 @@ public:
 					}
 				}
 
-				// Add a copy of the type stats to the build metadata as well
-				FString MetaFilename = FString::Printf(TEXT("%s/ShaderTypeStats-%s.csv"), *MetaDataDir, *PlatformName);
-				IFileManager::Get().Copy(*MetaFilename, *Filename);
+				// Add a copy of the type stats to the build metadata as well, potentially overriding name by cvar
+				FString StatsFileName = FString::Printf(TEXT("ShaderTypeStats-%s"), *PlatformName);
+				GetShaderFileNameOverride(StatsFileName, TEXT("r.Shaders.ShaderTypeStatsFileNameOverride"), PlatformFName);
+				FString MetaPath = FString::Printf(TEXT("%s/%s.csv"), *MetaDataDir, *StatsFileName);
+
+				if (IFileManager::Get().Copy(*MetaPath, *Filename) != COPY_OK)
+				{
+					UE_LOG(LogShaderLibrary, Warning, TEXT("Failed to copy from '%s' to '%s'"), *Filename, *MetaPath);
+				}
 			}
 			PlatformId++;
 		}
diff --git a/Engine/Source/Runtime/RenderCore/Private/ShaderCore.cpp b/Engine/Source/Runtime/RenderCore/Private/ShaderCore.cpp
--- a/Engine/Source/Runtime/RenderCore/Private/ShaderCore.cpp
+++ b/Engine/Source/Runtime/RenderCore/Private/ShaderCore.cpp
@@ -957,6 +957,38 @@ bool GetShaderSymbolPathOverride(FString& OutPathOverride, FName ShaderFormat)
 	return false;
 }
 
+bool GetShaderFileNameOverride(FString& OutFileNameOverride, const TCHAR* Cvar, FName ShaderFormat)
+{
+	FString OutFileName;
+	const FShaderSymbolSettingHelper SymbolFileNameOverride(Cvar);
+	if (SymbolFileNameOverride.GetString(OutFileName, ShaderFormat))
+	{
+		if (!OutFileName.IsEmpty())
+		{
+			// Allow the user to specify the location of the per-platform string.
+			OutFileName = OutFileName.Replace(TEXT("{Platform}"), *ShaderFormat.ToString(), ESearchCase::IgnoreCase);
+
+			// Allow adding name of the DLC if any to differentiate cooks, and one or more optional formatting dashes
+			FString DLCName;
+			FParse::Value(FCommandLine::Get(), TEXT("-dlcname="), DLCName);
+			OutFileName = OutFileName.Replace(TEXT("{?DLC-}"), DLCName.IsEmpty() ? TEXT("") : TEXT("-"), ESearchCase::IgnoreCase);
+			OutFileName = OutFileName.Replace(TEXT("{DLC}"), *DLCName, ESearchCase::IgnoreCase);
+
+			// Sanitize for safety
+			OutFileName.TrimStartAndEndInline();
+			OutFileName = FPaths::MakeValidFileName(OutFileName, TEXT('_'));
+		}
+	}
+
+	if (OutFileName.IsEmpty())
+	{
+		return false;
+	}
+
+	OutFileNameOverride = OutFileName;
+	return true;
+}
+
 EWriteShaderSymbols GetWriteShaderSymbolsOptions(FName ShaderFormat)
 {
 	static const FShaderSymbolSettingHelper WriteSymbolsZip(TEXT("r.Shaders.WriteSymbols.Zip"));
diff --git a/Engine/Source/Runtime/RenderCore/Public/ShaderCompilerCore.h b/Engine/Source/Runtime/RenderCore/Public/ShaderCompilerCore.h
--- a/Engine/Source/Runtime/RenderCore/Public/ShaderCompilerCore.h
+++ b/Engine/Source/Runtime/RenderCore/Public/ShaderCompilerCore.h
@@ -74,6 +74,10 @@ extern RENDERCORE_API bool ShouldWriteShaderSymbols(FName ShaderFormat);
 /** Returns true if the shader symbol path is overridden and OutPathOverride contains the override path. */
 extern RENDERCORE_API bool GetShaderSymbolPathOverride(FString& OutPathOverride, FName ShaderFormat);
 
+/** Returns true if the shader file name is overridden by the Cvar and OutFileNameOverride contains the override filename.
+    Returns false otherwise and leaves OutFileNameOverride unchanged */
+extern RENDERCORE_API bool GetShaderFileNameOverride(FString& OutFileNameOverride, const TCHAR* Cvar, FName ShaderFormat);
+
 /** Returns true if (external) shader symbols should be specific to each shader rather than be de-duplicated. */
 extern RENDERCORE_API bool ShouldAllowUniqueShaderSymbols(FName ShaderFormat);
 
