Get the content worker building with AutoRTFM + Merge Modules.

#rnx

diff --git a/Engine/Source/Developer/ShaderCompilerCommon/Public/DXCUtils.inl b/Engine/Source/Developer/ShaderCompilerCommon/Public/DXCUtils.inl
--- a/Engine/Source/Developer/ShaderCompilerCommon/Public/DXCUtils.inl
+++ b/Engine/Source/Developer/ShaderCompilerCommon/Public/DXCUtils.inl
@@ -79,7 +79,7 @@ static void DumpDebugBlobDetail(IDxcBlob* Blob, const TCHAR* BlobName, TRefCount
 	FPlatformMisc::LowLevelOutputDebugString(TEXT("\n\n"));
 }
 
-HRESULT RetrieveDebugNameAndBlob(TRefCountPtr<IDxcResult>& CompileResult, FString& OutDebugBlobName, IDxcBlob** OutDebugBlob, DxcShaderHash& ShaderHash)
+static HRESULT RetrieveDebugNameAndBlob(TRefCountPtr<IDxcResult>& CompileResult, FString& OutDebugBlobName, IDxcBlob** OutDebugBlob, DxcShaderHash& ShaderHash)
 {
 	check(CompileResult->HasOutput(DXC_OUT_SHADER_HASH));
 	{
