Shader Minifier, added Symbolic Dead Code Elimination (SDCE) from shallow program semantics

GPUTestBed PC cook: 16m -> 12m
(Substrate Adaptive, SM6-only)

Effectively eliminates dead stores on struct members, allowing backend compilers to eliminate intra-function side-effects prior to inlining. Written with performance in mind, on average SDCE costs the same as minification, e.g., a minification of 4ms usually implies SDCE of 4-5ms. Further gains could be realized by improving the tokenizer.

! As to why not the existing cross-compiler, my reasoning is that the goal is different. The cross-compiler is awesome, but has to play catch up with existing functionality (templates, unions, etc.), and requires a complete AST representation. This is specifically made to not require that, at the cost of complete semantics.

- Visible shader chunks (via the minifier) are shallow parsed, allowing for partial program representations. Parsing failures are conservatively handled by treating each identifier as "complex", preventing downstream DCE.
- Parsing produces a flat op-buffer (over AST, for performance) that's then iterated by a naive shunting yard approach to handle precedence. Though, it should be noted that the shallow parser only produces ops for the "things" we care about for SDCE. We can add as needed when we need more semantic awareness.
- Semantically we keep track of a lightweight symbol tree, marking tagged variables as used or not. Finally, we remove code sections that we can conservatively prove are not in use (*).

* SDCE currently does not track side-effects, so, we instead opt for SDCE on "tagged" types. To solve this, we need to keep track of data visibility when passed around. External? Everything's used. And, properly handle function qualifiers.

#FailedValidations Preflight

diff --git a/Engine/Source/Developer/ShaderCompilerCommon/Private/SDCE/SDCECommon.h b/Engine/Source/Developer/ShaderCompilerCommon/Private/SDCE/SDCECommon.h
new file mode 100644
--- /dev/null
+++ b/Engine/Source/Developer/ShaderCompilerCommon/Private/SDCE/SDCECommon.h
@@ -0,0 +1,69 @@
+﻿// Copyright Epic Games, Inc. All Rights Reserved.
+#pragma once
+
+#include "Containers/ContainerAllocationPolicies.h"
+#include "Misc/MemStack.h"
+#include "Misc/CString.h"
+
+/** Enable logging for debugging */
+#define UE_SHADER_SDCE_LOG_ALL        0
+#define UE_SHADER_SDCE_LOG_STRUCTURAL 0
+#define UE_SHADER_SDCE_LOG_COMPLEX    0
+
+/** If true, only populate symbols for tagged types, speeds up parsing but may miss nested type handling */
+#define UE_SHADER_SDCE_TAGGED_ONLY 0
+
+namespace UE::ShaderMinifier::SDCE
+{
+	/** Simple mutable view */
+	template<typename T>
+	struct TMutableStringView
+	{
+		TMutableStringView() : Begin(nullptr), Length(0)
+		{
+			
+		}
+
+		TMutableStringView(TString<T>& InView) : Begin(InView.GetCharArray().GetData()), Length(InView.Len())
+		{
+			
+		}
+
+		TMutableStringView(T* InBegin, int32 InLength) : Begin(InBegin), Length(InLength)
+		{
+			
+		}
+
+		TMutableStringView(T* InBegin) : TMutableStringView(InBegin, TCString<T>::Strlen(InBegin))
+		{
+			
+		}
+
+		T* begin() const { return Begin; }
+		T* end()   const { return Begin + Length; }
+		
+		T* Begin;
+		int32 Length;
+	};
+
+	/**
+	 * Note that while the minifier parses contents on its ansi sources,
+	 * SDCE is invoked during actual shader compilation which acts on wide
+	 * characters. Despite that, we implicitly assume that tokenization
+	 * segment lookups may operate on their ansi counterpart without
+	 * accounting for encoding differences (i.e., the wide character is ansi)
+	 */
+	using FParseCharType      = TCHAR;
+	using FParseViewType      = TMutableStringView<FParseCharType>;
+	using FParseConstViewType = TStringView<FParseCharType>;
+	using FParseStringType    = TString<FParseCharType>;
+	using FLookupCharType     = ANSICHAR;
+	using FLookupViewType     = TStringView<FLookupCharType>;
+
+	/**
+	 * Shader metadata prefix, e.g., UESHADERMETADATA_SDCE <name>
+	 */
+	static const TCHAR* ShaderMetadataPrefix = TEXT("SDCE ");
+	
+	using FMemStackSetAllocator = TSetAllocator<TSparseArrayAllocator<TMemStackAllocator<>, TMemStackAllocator<>>, TMemStackAllocator<>>;
+}
diff --git a/Engine/Source/Developer/ShaderCompilerCommon/Private/SDCE/SDCEContext.h b/Engine/Source/Developer/ShaderCompilerCommon/Private/SDCE/SDCEContext.h
new file mode 100644
--- /dev/null
+++ b/Engine/Source/Developer/ShaderCompilerCommon/Private/SDCE/SDCEContext.h
@@ -0,0 +1,313 @@
+﻿// Copyright Epic Games, Inc. All Rights Reserved.
+#pragma once
+
+#include "ShaderMinifier.h"
+#include "SDCEToken.h"
+#include "SDCEShallowParser.h"
+#include "SDCESymbolContext.h"
+#include "SDCESemanticContext.h"
+
+namespace UE::ShaderMinifier::SDCE
+{
+	/**
+	 * Symbolic Dead-Code-Elimination Context
+	 * Effectively a shallow parser with in-place semantics.
+	 */
+	class FContext
+	{
+		// In place parsing
+		friend class FShallowParser;
+		
+	public:
+		FContext(const TConstArrayView<FParseConstViewType>& DCESymbols) : SemanticCtx(SymbolCtx)
+		{
+			// Setup tagged type set
+			for (FParseConstViewType SymbolicDceStructName : DCESymbols)
+			{
+				TaggedTypes.Add(FConstSmallStringView::Text(SymbolicDceStructName.begin(), SymbolicDceStructName.Len()));
+			}
+		}
+
+		/**
+		 * Parse a given code chunk and generate relevant semantics
+		 */
+		void Analyze(const FParseViewType& Code)
+		{
+			// In-place parsing
+			Parser.ParseCode(Code, *this);
+		}
+
+		/**
+		 * Minify all parsed chunks in-place against their stores
+		 */
+		void MinifyInPlace()
+		{
+			// Eliminate in-order, lets us get away with a single text cursor
+			SymbolCtx.StoreSites.Sort([](const FMemberStoreSite& A, const FMemberStoreSite& B)
+			{
+				return A.AssignmentSite.Begin < B.AssignmentSite.Begin;
+			});
+
+			// Iteratively eliminate the realized loads of assignments we know are to be DCE'd
+			// i.e., If we have:
+			//   Data.A = 41;
+			//   Data.B = Data.A + 1;
+			//
+			// We know B is to be DCE'd, so we remove the realized loads on its RHS value span.
+			// After which A is known unrealized. And so on.
+			for (;;)
+			{
+				bool bMutatedRealizedLoad = false;
+
+				// Try to check for new unrealized loads
+				for (FMemberStoreSite& StoreSite : SymbolCtx.StoreSites)
+				{
+					check(StoreSite.AddressSymbol->bTagged);
+
+					// Still in use?
+					if (StoreSite.bVisited || StoreSite.AddressSymbol->Variable.RealizedLoadCount > 0)
+					{
+						continue;
+					}
+
+					// Mark as visited, fine before complex check since it's immutable
+					StoreSite.bVisited = true;
+
+					// Check if the symbol is considered "complex", i.e., the identifier was used in a
+					// code section that we couldn't parse or derive semantics from
+					StoreSite.bIsResolvedToComplex = SymbolCtx.ComplexBaseSet.Contains(StoreSite.AddressSymbol->ID);
+					if (StoreSite.bIsResolvedToComplex)
+					{
+						continue;
+					}
+
+					// Remove all the loads on the RHS
+					SemanticCtx.RemoveRealizedLoads(StoreSite.ValueSpan);
+					bMutatedRealizedLoad = true;
+				}
+
+				if (!bMutatedRealizedLoad)
+				{
+					break;
+				}
+			}
+
+			// Finally, eliminate!
+			for (const FMemberStoreSite& StoreSite : SymbolCtx.StoreSites)
+			{
+				check(StoreSite.AddressSymbol->bTagged);
+
+				// Ignore if the address has a realized load or complex in nature
+				if (StoreSite.AddressSymbol->Variable.RealizedLoadCount > 0 || StoreSite.bIsResolvedToComplex)
+				{
+					continue;
+				}
+				
+				// Replace with spaces while maintaining line numbers
+				for (int32 i = 0; i < StoreSite.AssignmentSite.Length; i++)
+				{
+					if (StoreSite.AssignmentSite.Begin[i] != '\n')
+					{
+						StoreSite.AssignmentSite.Begin[i] = ' ';
+					}
+				}
+			}
+		}
+	
+	private: /** Struct Visitors */
+		
+		void VisitStructPrologue(const FSmallStringView& ID)
+		{
+			// Push the stack
+			FSymbolScope& Scope = SymbolCtx.SymbolStack.Emplace_GetRef();
+
+			// Is this a tagged type?
+			const bool bTagged = TaggedTypes.Contains(ID);
+			
+	#if UE_SHADER_SDCE_LOG_STRUCTURAL
+			UE_LOG(LogTemp, Error, TEXT(">> ENTER %hs"), *FParseStringType::ConstructFromPtrSize(ID.Begin, ID.Length));
+	#endif // UE_SHADER_SDCE_LOG_STRUCTURAL
+
+#if UE_SHADER_SDCE_TAGGED_ONLY
+			if (!bTagged)
+			{
+				return;
+			}
+#endif // UE_SHADER_SDCE_TAGGED_ONLY
+			
+			// Assign composite of current stack
+			FSymbol* Symbol = SymbolCtx.CreateSymbol(ESymbolKind::Composite);
+			Symbol->ID = ID;
+			Symbol->bTagged = bTagged;
+			Scope.Composite = Symbol;
+	
+			// Add to parent stack
+			SymbolCtx.SymbolStack[SymbolCtx.SymbolStack.Num() - 2].Symbols.Add(ID, Symbol);
+		}
+		
+		void VisitStructEpilogue()
+		{
+			SymbolCtx.SymbolStack.Pop();
+		}
+	
+	private: /** Function Visitors */
+	
+		void VisitFunctionPrologue()
+		{
+			// Push the stack
+			SymbolCtx.SymbolStack.Emplace();
+	
+	#if UE_SHADER_SDCE_LOG_STRUCTURAL
+			FSmallStringView ID = SemanticCtx.GetOp(SemanticCtx.OpCount() - 1).ID;
+			UE_LOG(LogTemp, Error, TEXT(">> ENTER %hs"), *FParseStringType::ConstructFromPtrSize(ID.Begin, ID.Length));
+	#endif // UE_SHADER_SDCE_LOG_STRUCTURAL
+	
+			// Not semantically relevant for now
+			SemanticCtx.Empty();
+		}
+	
+		void VisitFunctionEpilogue()
+		{
+			SymbolCtx.SymbolStack.Pop();
+		}
+	
+	private: /** Expression Visitors */
+		
+		void VisitDeclaration()
+		{
+			FSemanticOp& Node = SemanticCtx.PushOp();
+			Node.Type = ESemanticOp::Declare;
+			Node.Precedence = EPrecedence::Declare;
+		}
+	
+		void VisitDrain(EPrecedence Precedence, const TCHAR* DebugReason)
+		{
+			FSemanticOp& DrainNode = SemanticCtx.PushOp();
+			DrainNode.Type = ESemanticOp::Drain;
+			DrainNode.Precedence = Precedence;
+			DrainNode.Drain.Count = 1;
+			DrainNode.Drain.DebugReason = DebugReason;
+		}
+	
+		void VisitType(const FSmallStringView& Type)
+		{
+			FSemanticOp& Node = SemanticCtx.PushOp();
+			Node.Type = ESemanticOp::Type;
+			Node.Precedence = EPrecedence::Primary;
+			Node.ID = Type;
+		}
+		
+		void VisitAlias()
+		{
+			FSymbolScope& Scope = SymbolCtx.SymbolStack.Last();
+	
+#if 0
+			// Add alias
+			if (FSymbol* Symbol = FindSymbol(Type.ID, ESymbolKind::Composite))
+			{
+				Scope.Symbols.Add(ID.ID, Symbol);
+			}
+#endif
+		}
+	
+		void VisitFramePrologue()
+		{
+			FSemanticOp& Node = SemanticCtx.PushOp();
+			Node.Type = ESemanticOp::FramePrologue;
+		}
+	
+		void VisitFrameEpilogue()
+		{
+			FSemanticOp& Node = SemanticCtx.PushOp();
+			Node.Type = ESemanticOp::FrameEpilogue;
+		}
+	
+		void VisitAggregateComposite()
+		{
+			FSemanticOp& Node = SemanticCtx.PushOp();
+			Node.Type = ESemanticOp::ID;
+			Node.Precedence = EPrecedence::Primary;
+			Node.ID = FSmallStringView{};
+		}
+	
+		void VisitPrimaryMemberAccess(const FSmallStringView& ID)
+		{
+			FSemanticOp& NodeAccess = SemanticCtx.PushOp();
+			NodeAccess.Type = ESemanticOp::Access;
+			NodeAccess.Precedence = EPrecedence::Access;
+			
+			FSemanticOp& Node = SemanticCtx.PushOp();
+			Node.Type = ESemanticOp::ID;
+			Node.Precedence = EPrecedence::Primary;
+			Node.ID = ID;
+		}
+	
+		void VisitPrimaryID(const FSmallStringView& ID)
+		{
+			FSemanticOp& Node = SemanticCtx.PushOp();
+			Node.Type = ESemanticOp::ID;
+			Node.Precedence = EPrecedence::Primary;
+			Node.ID = ID;
+		}
+	
+		void VisitComplexID(const FSmallStringView& ID)
+		{
+			// The parser has deemed this token complex, treat all accessors into said name as complex
+			SymbolCtx.AddComplexAccessor(ID);
+		}
+	
+		void VisitBinary(FToken Op)
+		{
+			FSemanticOp& Node = SemanticCtx.PushOp();
+			Node.Type = ESemanticOp::Binary;
+			Node.Precedence = Op.Type == ETokenType::BinaryEq ? EPrecedence::BinaryAssign : EPrecedence::BinaryGeneric;
+			Node.Binary.Op = Op;
+		}
+	
+	private: /** Statement Visitors */
+	
+		void VisitStatementPrologue(FParseCharType* Cursor)
+		{
+			CurrentStatementPrologueCursor = Cursor;
+			
+			// Just flush the previous contents
+			SemanticCtx.Empty();
+		}
+	
+		void VisitStatementEpilogue(FParseCharType* Cursor)
+		{
+#if UE_BUILD_DEBUG
+			// Useful for debugging
+			static uint32 DebugCounter = 0;
+			++DebugCounter;
+#endif // NDEBUG
+
+			// Evaluate the current op-stack
+			SemanticCtx.Evaluate(
+				FParseViewType(
+					CurrentStatementPrologueCursor,
+					Cursor - CurrentStatementPrologueCursor
+				)
+			);
+	
+	#if UE_SHADER_SDCE_LOG_ALL
+			UE_LOG(LogTemp, Error, TEXT("End of statement"));
+	#endif // UE_SHADER_SDCE_LOG_ALL
+		}
+
+	private:
+		/** Statement prologue cursor */
+		FParseCharType* CurrentStatementPrologueCursor = nullptr;
+
+	private:
+		/** Shared shallow parser */
+		FShallowParser Parser;
+
+		/** Contexts */
+		FSymbolContext   SymbolCtx;
+		FSemanticContext SemanticCtx;
+
+		/** Actual set of types we're minifying */
+		TSet<FConstSmallStringView, FConstSmallStringKeyFuncs, FMemStackSetAllocator> TaggedTypes;
+	};
+}
diff --git a/Engine/Source/Developer/ShaderCompilerCommon/Private/SDCE/SDCEOpBuffer.h b/Engine/Source/Developer/ShaderCompilerCommon/Private/SDCE/SDCEOpBuffer.h
new file mode 100644
--- /dev/null
+++ b/Engine/Source/Developer/ShaderCompilerCommon/Private/SDCE/SDCEOpBuffer.h
@@ -0,0 +1,87 @@
+﻿// Copyright Epic Games, Inc. All Rights Reserved.
+#pragma once
+
+#include "Misc/MemStack.h"
+
+namespace UE::ShaderMinifier::SDCE
+{
+	template<typename T, uint32 ChunkSize>
+	struct TOpBuffer
+	{
+		/** Emplace a new op */
+		T& EmplaceUninitialized()
+		{
+			int32 ChunkIndex = Count / ChunkSize;
+			int32 OpIndex    = Count % ChunkSize;
+	
+			if (!Chunks.IsValidIndex(ChunkIndex))
+			{
+				for (int32 i = Chunks.Num(); i <= ChunkIndex; i++)
+				{
+					Chunks.Add(static_cast<Chunk*>(FMemStack::Get().Alloc(sizeof(Chunk), alignof(Chunk))));
+				}
+			}
+	
+			T& Op = Chunks[ChunkIndex]->Ops[OpIndex];
+			Count++;
+			return Op;
+		}
+
+		/** Add a new chunk value */
+		void Add(const T& Value)
+		{
+			new (&EmplaceUninitialized()) T(Value);
+		}
+
+		/** Get an op by its index */
+		const T& operator[](int32 Index) const
+		{
+			int32 ChunkIndex = Index / ChunkSize;
+			int32 OpIndex    = Index % ChunkSize;
+			return Chunks[ChunkIndex]->Ops[OpIndex];
+		}
+
+		/** Get an op by its index */
+		T& operator[](int32 Index)
+		{
+			int32 ChunkIndex = Index / ChunkSize;
+			int32 OpIndex    = Index % ChunkSize;
+			return Chunks[ChunkIndex]->Ops[OpIndex];
+		}
+
+		/** Any ops? */
+		bool IsEmpty() const
+		{
+			return Count == 0;
+		}
+
+		/** Set the number of ops */
+		void SetNumNoAlloc(int32 InCount)
+		{
+			check(InCount >= 0 && InCount <= static_cast<int32>(Chunks.Num() * ChunkSize));
+			Count = InCount;
+		}
+	
+		/** Clear this buffer */
+		void Empty()
+		{
+			Count = 0;
+		}
+	
+		/** Number of available ops */
+		int32 Num() const
+		{
+			return Count;
+		}
+
+	private:
+		struct Chunk
+		{
+			T Ops[ChunkSize];
+		};
+	
+		TArray<Chunk*, TMemStackAllocator<>> Chunks;
+		
+		int32 Count = 0;
+	};
+}
diff --git a/Engine/Source/Developer/ShaderCompilerCommon/Private/SDCE/SDCESemanticContext.h b/Engine/Source/Developer/ShaderCompilerCommon/Private/SDCE/SDCESemanticContext.h
new file mode 100644
--- /dev/null
+++ b/Engine/Source/Developer/ShaderCompilerCommon/Private/SDCE/SDCESemanticContext.h
@@ -0,0 +1,664 @@
+﻿// Copyright Epic Games, Inc. All Rights Reserved.
+#pragma once
+
+#include "SDCECommon.h"
+#include "SDCEShallowParser.h"
+#include "SDCESymbolContext.h"
+#include "SDCEOpBuffer.h"
+#include "SDCEToken.h"
+
+namespace UE::ShaderMinifier::SDCE
+{
+	enum class ESemanticOp
+	{
+		ID,
+		Type,
+		Access,
+		Binary,
+		Declare,
+		Drain,
+		FramePrologue,
+		FrameEpilogue,
+	};
+	
+	struct FSemanticOp
+	{
+		ESemanticOp Type;
+		EPrecedence Precedence;
+
+		/** Payload */
+		union
+		{
+			/** ESemanticOp::ID */
+			FSmallStringView ID;
+				
+			/** ESemanticOp::Binary */
+			struct
+			{
+				FToken Op;
+			} Binary;
+	
+			/** ESemanticOp::Drain */
+			struct
+			{
+				int32 Count;
+				const TCHAR* DebugReason;
+			} Drain;
+		};
+	};
+	
+	static_assert(sizeof(FSemanticOp) == 24, "Unexpected FSemanticOp packing");
+
+	/**
+	 * Simple semantic evaluator, reduces a set of semantic operations by their natural precedence
+	 * with a naive shunting yard approach. This was picked over an "AST" like structure to keep
+	 * things as lightweight as possible.
+	 */
+	class FSemanticContext
+	{
+	public:
+		FSemanticContext(FSymbolContext& SymCtx) : SymCtx(SymCtx)
+		{
+			
+		}
+
+		/** Clear the context */
+		void Empty()
+		{
+			OpBuffer.Empty();
+		}
+
+		/** Create a new op */
+		FSemanticOp& PushOp()
+		{
+			return OpBuffer.EmplaceUninitialized();
+		}
+
+		const FSemanticOp& GetOp(int32 Index) const
+		{
+			return OpBuffer[Index];
+		}
+
+		int32 OpCount() const
+		{
+			return OpBuffer.Num();
+		}
+
+		/** Evaluate all pushed ops */
+		void Evaluate(const FParseViewType& StatementRegion)
+		{
+			CurrentStatementRegion = StatementRegion;
+
+			// Early out
+			if (OpBuffer.IsEmpty())
+			{
+				return;
+			}
+
+			const int32 ValueHistoryBegin = ValueHistory.Num();
+	
+			// Single pass tape evaluation
+			for (int32 OpIndex = 0; OpIndex < OpBuffer.Num(); OpIndex++)
+			{
+				const FSemanticOp& Buffer = OpBuffer[OpIndex];
+				switch (Buffer.Type)
+				{
+					case ESemanticOp::ID:
+					{
+						FStackValue Value;
+						Value.ID = Buffer.ID;
+						Value.Span = FValueSpan(ValueHistory.Num(), ValueHistory.Num() + 1);
+						PushValue(Value);
+						break;
+					}
+					case ESemanticOp::Type:
+					{
+						FStackValue Value;
+						Value.ID = Buffer.ID;
+						Value.Span = FValueSpan(ValueHistory.Num(), ValueHistory.Num() + 1);
+						PushValue(Value);
+						break;
+					}
+					case ESemanticOp::Access:
+					{
+						FReduceOp Op;
+						Op.Type = ESemanticOp::Access;
+						Op.Precedence = EPrecedence::Access;
+						PushOp(Op);
+						break;
+					}
+					case ESemanticOp::Binary:
+					{
+						FReduceOp Op;
+						Op.Type = ESemanticOp::Binary;
+						Op.Precedence = Buffer.Precedence;
+						Op.Binary.Op = Buffer.Binary.Op;
+						PushOp(Op);
+						break;
+					}
+					case ESemanticOp::Declare:
+					{
+						FReduceOp Op;
+						Op.Type = ESemanticOp::Declare;
+						Op.Precedence = EPrecedence::Declare;
+						PushOp(Op);
+						break;
+					}
+					case ESemanticOp::Drain:
+					{
+						FReduceOp Op;
+						Op.Type = ESemanticOp::Drain;
+						Op.Precedence = Buffer.Precedence;
+						Op.Drain.Count = Buffer.Drain.Count;
+						PushOp(Op);
+						break;
+					}
+					case ESemanticOp::FramePrologue:
+					{
+						FFrameHead Head;
+						Head.OpHead = OpStack.Num();
+						Head.ValueHead = ValueStack.Num();
+						FrameStack.Add(Head);
+						break;
+					}
+					case ESemanticOp::FrameEpilogue:
+					{
+						FFrameHead Head = FrameStack.Pop();
+
+						// Number of operands we need to drain
+						int32 PendingOps = OpStack.Num() - Head.OpHead;
+						check(PendingOps >= 0);
+
+						// Keep draining
+						for (int32 i = 0; i < PendingOps; i++)
+						{
+							Reduce();
+						}
+	
+						check(OpStack.Num() == Head.OpHead);
+						break;
+					}
+				}
+			}
+
+			// Reduce all pending ops
+			while (!OpStack.IsEmpty())
+			{
+				Reduce();
+			}
+
+			// Go through all accessed values, and treat all read sites as loaded
+			for (int32 i = ValueHistoryBegin; i < ValueHistory.Num(); i++)
+			{
+				const FStackValue& Value = ValueHistory[i];
+				if (!Value.bIsReadSite)
+				{
+					continue;
+				}
+
+				// Only tracking for variables currently
+				if (Value.Symbol && Value.Symbol->Kind == ESymbolKind::Variable)
+				{
+					Value.Symbol->Variable.RealizedLoadCount++;
+				}
+			}
+			
+			// Mismatched drain, somewhere
+			// Note that statements may reduce to a single value, that's completely fine
+			check(ValueStack.Num() <= 1);
+
+			// Cleanup
+			OpBuffer.Empty();
+			ValueStack.Empty();
+			OpStack.Empty();
+			FrameStack.Empty();
+	
+	#if UE_SHADER_SDCE_LOG_ALL
+			UE_LOG(LogTemp, Error, TEXT("End of statement"));
+	#endif // UE_SHADER_SDCE_LOG_ALL
+		}
+
+		/** Remove all the realized from a span of values */
+		void RemoveRealizedLoads(const FValueSpan& Span)
+		{
+			for (int32 i = Span.Begin; i < Span.End; i++)
+			{
+				FStackValue& Value = ValueHistory[i];
+				if (!Value.bIsReadSite)
+				{
+					continue;
+				}
+
+				// Only tracking for variables currently
+				if (Value.Symbol && Value.Symbol->Kind == ESymbolKind::Variable)
+				{
+					Value.Symbol->Variable.RealizedLoadCount--;
+				}
+			}
+		}
+
+	private:
+		/** Internal reduction op */
+		struct FReduceOp
+		{
+			ESemanticOp Type;
+			EPrecedence Precedence;
+
+			/** Payload */
+			union
+			{
+				struct
+				{
+					FToken Op;
+				} Binary;
+					
+				struct
+				{
+					int32 Count;
+				} Drain;
+			};
+		};
+	
+		static_assert(sizeof(FReduceOp) == 24, "Unexpected FReduceOp packing");
+
+	private:
+		/** Reduce the top most op */
+		void Reduce()
+		{
+			FReduceOp Op = OpStack.Pop();
+			switch (Op.Type)
+			{
+				default:
+				{
+					checkNoEntry();
+				}
+				case ESemanticOp::Access:
+				{					
+					// Stack stack
+					FStackValue RHS = ValueStack.Pop();
+					FStackValue LHS = ValueStack.Pop();
+
+					// Resolve it
+					FSymbol* LHSSymbol = ResolveValueSymbol(LHS);
+					if (!LHSSymbol)
+					{
+#if UE_SHADER_SDCE_LOG_COMPLEX
+						UE_LOG(LogTemp, Error, TEXT("LHS Complex: Failed to resolve"));
+#endif // UE_SHADER_SDCE_LOG_COMPLEX
+
+						// Failed to resolve, treat as complex
+						PushUnresolvedValue(LHS.Span | RHS.Span);
+						SymCtx.AddComplexAccessor(RHS.ID);
+						break;
+					}
+
+					// Unwrap variables
+					if (LHSSymbol->Kind == ESymbolKind::Variable)
+					{
+						// If nested, treat LHS as loaded
+						if (LHSSymbol->bTagged)
+						{
+							// TODO: Is this correct? I can't remember why I added this...
+							LHSSymbol->Variable.RealizedLoadCount++;
+						}
+
+						// Typeless? Treat as complex
+						LHSSymbol = LHSSymbol->Variable.Type;
+						if (!LHSSymbol)
+						{
+#if UE_SHADER_SDCE_LOG_COMPLEX
+							UE_LOG(LogTemp, Error, TEXT("LHS Complex: Failed to resolve variable type"));
+#endif // UE_SHADER_SDCE_LOG_COMPLEX
+							
+							PushUnresolvedValue(LHS.Span | RHS.Span);
+							SymCtx.AddComplexAccessor(RHS.ID);
+							break;
+						}
+					}
+
+					// Find the contained symbol, may fail
+					FSymbol* ContainedSymbol = SymCtx.LinearCompositeFindSymbol(LHSSymbol, RHS.ID);
+
+					// Setup 
+					FStackValue Value;
+					Value.Symbol = ContainedSymbol;
+					Value.Span = FValueSpan(ValueHistory.Num(), ValueHistory.Num() + 1);
+					Value.Span = LHS.Span | Value.Span;
+					Value.bIsReadSite = true;
+					PushValue(Value);
+
+#if UE_SHADER_SDCE_LOG_ALL
+					if (LHSSymbol && ContainedSymbol)
+					{
+						UE_LOG(LogTemp, Error, TEXT("Access %hs . %hs"), *FParseStringType::ConstructFromPtrSize(LHSSymbol->ID.Begin, LHSSymbol->ID.Length), *FParseStringType::ConstructFromPtrSize(ContainedSymbol->ID.Begin, ContainedSymbol->ID.Length));
+					}
+#endif // UE_SHADER_SDCE_LOG_ALL
+
+					break;
+				}
+				case ESemanticOp::Binary:
+				{
+					FStackValue RHS = ValueStack.Pop();
+					FStackValue LHS = ValueStack.Pop();
+
+					// Resolve the operands (RHS for load tracking)
+					FSymbol* LHSSymbol = ResolveValueSymbol(LHS);
+					FSymbol* RHSSymbol = ResolveValueSymbolReadSite(RHS);
+
+					// If non-assignment, just ignore
+					if (Op.Binary.Op.Type != ETokenType::BinaryEq)
+					{
+						PushUnresolvedValue(LHS.Span | RHS.Span);
+						break;
+					}
+
+					// Did we resolve the LHS?
+					if (LHSSymbol)
+					{
+						// If not tagged, ignore
+						if (!LHSSymbol->bTagged)
+						{
+							PushUnresolvedValue(LHS.Span | RHS.Span);
+							break;
+						}
+
+						// For now, only DCE simple store sites
+						// Will be expanded on in the future
+						if (IsTrivialNonCompoundAssignment())
+						{
+							FMemberStoreSite& Site = SymCtx.StoreSites.Emplace_GetRef();
+							Site.AddressSymbol = LHSSymbol;
+							Site.ValueSymbol = RHSSymbol;
+							Site.AddressSpan = LHS.Span;
+							Site.ValueSpan = RHS.Span;
+							Site.AssignmentSite.Begin  = CurrentStatementRegion.Begin;
+							Site.AssignmentSite.Length = static_cast<uint16_t>(CurrentStatementRegion.Length);
+						}
+					}
+
+#if UE_SHADER_SDCE_LOG_ALL
+					if (LHSSymbol && RHSSymbol)
+					{
+						UE_LOG(LogTemp, Error, TEXT("BINARY %hs = %hs"), *FParseStringType::ConstructFromPtrSize(LHSSymbol->ID.Begin, LHSSymbol->ID.Length), *FParseStringType::ConstructFromPtrSize(RHS.ID.Begin, RHS.ID.Length));
+					}
+#endif // UE_SHADER_SDCE_LOG_ALL
+
+					// Mark all variables on the LHS as conditional
+					// Need to check for side effects, obviously
+					for (int32 i = LHS.Span.Begin; i < LHS.Span.End; i++)
+					{
+						FStackValue& Value = ValueHistory[i];
+						Value.bIsReadSite = false;
+					}
+
+					// TODO: Add comment about side effects
+					LHS.bIsReadSite = false;
+
+					// Assignment propagates LHS
+					PushValue(LHS);
+					break;
+				}
+				case ESemanticOp::Declare:
+				{
+					FSymbolScope& Scope = SymCtx.SymbolStack.Last();
+
+					// Stack stack
+					FStackValue ID   = ValueStack.Pop();
+					FStackValue Type = ValueStack.Pop();
+
+					// Find the type
+					FSymbol* TypeSymbol = SymCtx.FindSymbol(Type.ID);
+
+					// Tagged status, for now we only tag declarations inside of composites
+					// Function local declarations require qualifier knowledge, which is a TODO
+					bool bTagged = false;
+					if (Scope.Composite)
+					{
+						bTagged |= Scope.Composite->bTagged;
+						
+						if (TypeSymbol)
+						{
+							bTagged |= TypeSymbol->bTagged;
+						}
+					}
+
+#if UE_SHADER_SDCE_TAGGED_ONLY
+					if (!bTagged)
+					{
+						PushUnresolvedValue();
+						break;
+					}
+#endif // UE_SHADER_SDCE_TAGGED_ONLY
+
+					// Add to current composite
+					FSymbol* Symbol = SymCtx.CreateSymbol(ESymbolKind::Variable);
+					Symbol->ID = ID.ID;
+					Symbol->Variable.Type = TypeSymbol;
+					Symbol->bTagged = bTagged;
+
+					// Current scope
+					if (Scope.Composite)
+					{
+						Scope.Composite->Composite.ContainedSymbols.Add(Symbol);
+					}
+
+					// Add lookup to stack
+					Scope.Symbols.Add(ID.ID, Symbol);
+
+#if UE_SHADER_SDCE_LOG_ALL
+					UE_LOG(LogTemp, Error, TEXT("DECLARE %hs %hs"), *FParseStringType::ConstructFromPtrSize(Type.ID.Begin, Type.ID.Length), *FParseStringType::ConstructFromPtrSize(ID.ID.Begin, ID.ID.Length));
+#endif // UE_SHADER_SDCE_LOG_ALL
+
+					FStackValue Value;
+					Value.Symbol = Symbol;
+					Value.Span = FValueSpan(ValueHistory.Num(), ValueHistory.Num() + 1);
+					PushValue(Value);
+					break;
+				}
+				case ESemanticOp::Drain:
+				{
+					FValueSpan Span;
+
+					// On shallow parsing drains, just treat the symbol, if any, as loaded
+					for (int32 i = 0; i < Op.Drain.Count; i++)
+					{
+						FStackValue Value = ValueStack.Pop();
+
+						// Resolve the symbol and treat it as a read site
+						ResolveValueSymbolReadSite(Value);
+
+						// Keep track of the value span
+						if (i == 0)
+						{
+							Span = Value.Span;
+						}
+						else
+						{
+							Span |= Value.Span;
+						}
+					}
+
+					// We always drain to the left, so, have the left hand side inherit the value span
+					if (!ValueStack.IsEmpty())
+					{
+						ValueStack.Last().Span |= Span;
+					}
+					break;
+				}
+			}
+		}
+
+		/** Push a new op */
+		void PushOp(const FReduceOp& Op)
+		{
+			// We can only reduce up to the current frame boundary
+			int32 FrameBoundary = !FrameStack.IsEmpty() ? FrameStack.Last().OpHead : 0;
+
+			// Keep reducing until at boundary
+			while (!OpStack.IsEmpty() && OpStack.Num() > FrameBoundary)
+			{				
+				FReduceOp& Last = OpStack.Last();
+
+				// If less or at, we can reduce
+				// TODO: Doesn't account for LHS/RHS assoc. not needed for now
+				if (Last.Precedence > Op.Precedence)
+				{
+					break;
+				}
+				
+				Reduce();
+			}
+
+			OpStack.Add(Op);
+		}
+
+	private:
+		bool IsTrivialNonCompoundAssignment()
+		{
+			int32 BaseFrameAssignments = 0;
+			int32 FrameCounter         = 0;
+			
+			for (int32 OpIndex = 0; OpIndex < OpBuffer.Num(); OpIndex++)
+			{
+				const FSemanticOp& Buffer = OpBuffer[OpIndex];
+
+				// Handle scoping
+				switch (Buffer.Type)
+				{
+					default:
+					{
+						break;
+					}
+					case ESemanticOp::FramePrologue:
+					{
+						FrameCounter++;
+						continue;
+					}
+					case ESemanticOp::FrameEpilogue:
+					{
+						FrameCounter--;
+						continue;
+					}
+				}
+
+				// Ignore all nested scopes
+				if (FrameCounter > 0)
+				{
+					continue;
+				}
+
+				// If there's anything of greater precedence, non-trivial
+				if (Buffer.Precedence > EPrecedence::BinaryAssign)
+				{
+					return false;
+				}
+
+				// Assignment?
+				BaseFrameAssignments += (Buffer.Precedence == EPrecedence::BinaryAssign);
+			}
+
+			// Ignoring nested assignment, for now
+			return BaseFrameAssignments == 1;
+		}
+
+	private:
+		/** Value on the stack */
+		struct FStackValue
+		{
+			/** Optional, resolved symbol */
+			FSymbol* Symbol = nullptr;
+
+			/** If unresolved, identifer of this value */
+			FSmallStringView ID{};
+
+			/** The effectively value span this represents */
+			FValueSpan Span;
+
+			/** Should this be treated as a read site? */
+			bool bIsReadSite = false;
+		};
+
+		/** Frame counters for drains */
+		struct FFrameHead
+		{
+			int32 OpHead    = 0;
+			int32 ValueHead = 0;
+		};
+
+		/** Helper, get the symbol of a stack value */
+		FSymbol* ResolveValueSymbol(const FStackValue& Value) const
+		{
+			if (Value.Symbol)
+			{
+				return Value.Symbol;
+			}
+
+			if (Value.ID.Begin)
+			{
+				return SymCtx.FindSymbol(Value.ID);
+			}
+
+			return nullptr;
+		}
+
+		/** Helper, get the symbol of a stack value */
+		FSymbol* ResolveValueSymbolReadSite(const FStackValue& Value)
+		{
+			if (Value.Symbol)
+			{
+				return Value.Symbol;
+			}
+
+			if (Value.ID.Begin)
+			{
+				auto Symbol = SymCtx.FindSymbol(Value.ID);
+
+				if (Symbol)
+				{
+					ValueHistory.Add(FStackValue {
+						.Symbol = Symbol,
+						.Span = Value.Span,
+						.bIsReadSite = true
+					});
+				}
+
+				return Symbol;
+			}
+
+			return nullptr;
+		}
+		
+		/** Helper, push a value to the stack */
+		void PushValue(const FStackValue& Value)
+		{
+			ValueStack.Add(Value);
+			ValueHistory.Add(Value);
+		}
+
+		/** Helper, push an unresolved value */
+		void PushUnresolvedValue(const FValueSpan& Span)
+		{
+			PushValue(FStackValue {
+				.Span = Span
+			});
+		}
+
+	private:
+		/** Current op buffer */
+		TOpBuffer<FSemanticOp, 128> OpBuffer;
+
+		/** Evaluation stacks */
+		TArray<FStackValue, TInlineAllocator<64, TMemStackAllocator<>>> ValueStack;
+		TArray<FReduceOp,   TInlineAllocator<64, TMemStackAllocator<>>> OpStack;
+		TArray<FFrameHead,  TInlineAllocator<64, TMemStackAllocator<>>> FrameStack;
+
+		/** Evaluation value history, not a stack */
+		TOpBuffer<FStackValue, 1024> ValueHistory;
+
+		/** Current statement code region */
+		FParseViewType CurrentStatementRegion;
+
+	private:
+		FSymbolContext& SymCtx;
+	};
+}
diff --git a/Engine/Source/Developer/ShaderCompilerCommon/Private/SDCE/SDCEShallowParser.h b/Engine/Source/Developer/ShaderCompilerCommon/Private/SDCE/SDCEShallowParser.h
new file mode 100644
--- /dev/null
+++ b/Engine/Source/Developer/ShaderCompilerCommon/Private/SDCE/SDCEShallowParser.h
@@ -0,0 +1,1435 @@
+﻿// Copyright Epic Games, Inc. All Rights Reserved.
+#pragma once
+
+#include "SDCECommon.h"
+#include "SDCETokenizer.h"
+#include "SDCEOpBuffer.h"
+
+namespace UE::ShaderMinifier::SDCE
+{
+	/** Standardized precedence for all operations */
+	enum class EPrecedence
+	{
+		None,
+		Primary,
+		Access,
+		Declare,
+		BinaryGeneric,
+		BinaryAssign,
+		Statement
+	};
+
+	/**
+	 * As per its name, the *shallow* parser doesn't require full program understanding
+	 * it instead attempts to parse what it can, and treat the rest as "complex". It's up
+	 * to the visitor to interpret this as it sees fit.
+	 */
+	class FShallowParser
+	{
+	public:
+		FShallowParser()
+		{
+			// Let's set up a standardized set of templated types. This is not complete, nor does it aim to be,
+			// just to get "enough" stuff through.
+			TemplatedTypes.Add(FConstSmallStringView::Text(TEXT("texture")));
+			TemplatedTypes.Add(FConstSmallStringView::Text(TEXT("Texture1D")));
+			TemplatedTypes.Add(FConstSmallStringView::Text(TEXT("Texture1DArray")));
+			TemplatedTypes.Add(FConstSmallStringView::Text(TEXT("Texture1D_Array")));
+			TemplatedTypes.Add(FConstSmallStringView::Text(TEXT("Texture2D")));
+			TemplatedTypes.Add(FConstSmallStringView::Text(TEXT("Texture2DArray")));
+			TemplatedTypes.Add(FConstSmallStringView::Text(TEXT("Texture2D_Array")));
+			TemplatedTypes.Add(FConstSmallStringView::Text(TEXT("Texture2DMS")));
+			TemplatedTypes.Add(FConstSmallStringView::Text(TEXT("MS_Texture2D")));
+			TemplatedTypes.Add(FConstSmallStringView::Text(TEXT("Texture2DMSArray")));
+			TemplatedTypes.Add(FConstSmallStringView::Text(TEXT("MS_Texture2D_Array")));
+			TemplatedTypes.Add(FConstSmallStringView::Text(TEXT("Texture3D")));
+			TemplatedTypes.Add(FConstSmallStringView::Text(TEXT("TextureCube")));
+			TemplatedTypes.Add(FConstSmallStringView::Text(TEXT("TextureCubeArray")));
+			TemplatedTypes.Add(FConstSmallStringView::Text(TEXT("TextureCube_Array")));
+			TemplatedTypes.Add(FConstSmallStringView::Text(TEXT("Buffer")));
+			TemplatedTypes.Add(FConstSmallStringView::Text(TEXT("DataBuffer")));
+			TemplatedTypes.Add(FConstSmallStringView::Text(TEXT("AppendStructuredBuffer")));
+			TemplatedTypes.Add(FConstSmallStringView::Text(TEXT("AppendRegularBuffer")));
+			TemplatedTypes.Add(FConstSmallStringView::Text(TEXT("ByteAddressBuffer")));
+			TemplatedTypes.Add(FConstSmallStringView::Text(TEXT("ByteBuffer")));
+			TemplatedTypes.Add(FConstSmallStringView::Text(TEXT("ConsumeStructuredBuffer")));
+			TemplatedTypes.Add(FConstSmallStringView::Text(TEXT("ConsumeRegularBuffer")));
+			TemplatedTypes.Add(FConstSmallStringView::Text(TEXT("RWBuffer")));
+			TemplatedTypes.Add(FConstSmallStringView::Text(TEXT("RW_DataBuffer")));
+			TemplatedTypes.Add(FConstSmallStringView::Text(TEXT("RWByteAddressBuffer")));
+			TemplatedTypes.Add(FConstSmallStringView::Text(TEXT("RW_ByteBuffer")));
+			TemplatedTypes.Add(FConstSmallStringView::Text(TEXT("RWStructuredBuffer")));
+			TemplatedTypes.Add(FConstSmallStringView::Text(TEXT("RW_RegularBuffer")));
+			TemplatedTypes.Add(FConstSmallStringView::Text(TEXT("RWTexture1D")));
+			TemplatedTypes.Add(FConstSmallStringView::Text(TEXT("RW_Texture1D")));
+			TemplatedTypes.Add(FConstSmallStringView::Text(TEXT("RWTexture1DArray")));
+			TemplatedTypes.Add(FConstSmallStringView::Text(TEXT("RW_Texture1D_Array")));
+			TemplatedTypes.Add(FConstSmallStringView::Text(TEXT("RWTexture2D")));
+			TemplatedTypes.Add(FConstSmallStringView::Text(TEXT("RW_Texture2D")));
+			TemplatedTypes.Add(FConstSmallStringView::Text(TEXT("RWTexture2DArray")));
+			TemplatedTypes.Add(FConstSmallStringView::Text(TEXT("RW_Texture2D_Array")));
+			TemplatedTypes.Add(FConstSmallStringView::Text(TEXT("RasterizerOrderedTexture2D")));
+			TemplatedTypes.Add(FConstSmallStringView::Text(TEXT("RWTexture3D")));
+			TemplatedTypes.Add(FConstSmallStringView::Text(TEXT("RW_Texture3D")));
+			TemplatedTypes.Add(FConstSmallStringView::Text(TEXT("StructuredBuffer")));
+			TemplatedTypes.Add(FConstSmallStringView::Text(TEXT("RegularBuffer")));
+			TemplatedTypes.Add(FConstSmallStringView::Text(TEXT("ConstantBuffer")));
+			TemplatedTypes.Add(FConstSmallStringView::Text(TEXT("RaytracingAccelerationStructure")));
+			TemplatedTypes.Add(FConstSmallStringView::Text(TEXT("TriangleStream")));
+			TemplatedTypes.Add(FConstSmallStringView::Text(TEXT("ConstantBuffer")));
+			TemplatedTypes.Add(FConstSmallStringView::Text(TEXT("vector")));
+		}
+
+		/**
+		 * Parse a code chunk, visitor statically invoked
+		 */
+		template<typename T>
+		void ParseCode(const FParseViewType& Code, T&& Visitor)
+		{
+			// Useful for debugging
+	#if 0
+			Type = ECodeChunkType::Function;
+			Code = R"()";
+	#endif // 0
+
+			// Setup tokenizer
+			FShallowTokenizer Tk;
+			Tk.Begin = Code.Begin;
+			Tk.End = Code.Begin + Code.Length;
+			Tk.Next();
+
+			// Just keep trying to parse as statements
+			while (!Tk.IsEOS())
+			{
+				ParseStatement(Tk);
+			}
+
+			// Succeeded, propagate visitation
+			for (int32 i = 0; i < OpBuffer.Num(); i++)
+			{
+				const FVisitOp& Op = OpBuffer[i];
+				switch (Op.Type)
+				{
+					default:
+					{
+						checkNoEntry();
+						break;
+					}
+					case EVisitOp::ID: {
+						Visitor.VisitPrimaryID(Op.ID);
+						break;
+					}
+					case EVisitOp::ComplexID: {
+						Visitor.VisitComplexID(Op.ID);
+						break;
+					}
+					case EVisitOp::Type: {
+						Visitor.VisitType(Op.ID);
+						break;
+					}
+					case EVisitOp::Binary: {
+						Visitor.VisitBinary(Op.Binary.Token);
+						break;
+					}
+					case EVisitOp::MemberAccess: {
+						Visitor.VisitPrimaryMemberAccess(Op.ID);
+						break;
+					}
+					case EVisitOp::AggregateComposite: {
+						Visitor.VisitAggregateComposite();
+						break;
+					}
+					case EVisitOp::Declaration: {
+						Visitor.VisitDeclaration();
+						break;
+					}
+					case EVisitOp::StatementPrologue: {
+						Visitor.VisitStatementPrologue(Op.Statement.Cursor);
+						break;
+					}
+					case EVisitOp::StatementEpilogue: {
+						Visitor.VisitStatementEpilogue(Op.Statement.Cursor);
+						break;
+					}
+					case EVisitOp::StructPrologue: {
+						Visitor.VisitStructPrologue(Op.Struct.ID);
+						break;
+					}
+					case EVisitOp::StructEpilogue: {
+						Visitor.VisitStructEpilogue();
+						break;
+					}
+					case EVisitOp::FunctionPrologue: {
+						Visitor.VisitFunctionPrologue();
+						break;
+					}
+					case EVisitOp::FunctionEpilogue: {
+						Visitor.VisitFunctionEpilogue();
+						break;
+					}
+					case EVisitOp::FramePrologue: {
+						Visitor.VisitFramePrologue();
+						break;
+					}
+					case EVisitOp::FrameEpilogue: {
+						Visitor.VisitFrameEpilogue();
+						break;
+					}
+					case EVisitOp::Alias: {
+						Visitor.VisitAlias();
+						break;
+					}
+					case EVisitOp::Drain: {
+						Visitor.VisitDrain(Op.Drain.Type, Op.Drain.Reason);
+						break;
+					}
+				}
+			}
+
+			OpBuffer.Empty();
+			
+			Tk.OpHead = TokenizerInvalidOpHead;
+		}
+
+	private:
+		/** Helper, appends prologue/epilogue around a statement */
+		template<typename F>
+		void WithStatement(FShallowTokenizer& Tk, F&& functor)
+		{
+			OpBuffer.Add(FVisitOp { .Type = EVisitOp::StatementPrologue, .Statement = { .Cursor = Tk.Token.Begin } });
+			
+			functor();
+
+			OpBuffer.Add(FVisitOp { .Type = EVisitOp::StatementEpilogue, .Statement = { .Cursor = Tk.Token.Begin + Tk.Token.Length } });
+		}
+
+		void ParseStruct(FShallowTokenizer& Tk)
+		{
+			bool bIsTemplated = SkipTemplate(Tk);
+			SkipAnnotations(Tk);
+			
+			check(Tk.IsConsume(ETokenType::Struct));
+			FToken ID = Tk.Next();
+
+			// If templated, keep track of it
+			if (bIsTemplated)
+			{
+				TemplatedTypes.Add(FSmallStringView::Get(ID));
+			}
+
+			OpBuffer.Add(FVisitOp { .Type = EVisitOp::StructPrologue, .Struct = { FSmallStringView::Get(ID) } });
+
+			if (!Tk.Is(ETokenType::BodyOpen) && !Tk.Is(ETokenType::EndOfStatement))
+			{
+				// TODO: We need structural "skipping" with complex handlers
+				Tk.SkipUntil(TEXT("{;"));
+			}
+
+			if (Tk.Is(ETokenType::BodyOpen))
+			{
+				ParseBody(Tk);
+			}
+
+			OpBuffer.Add(FVisitOp { .Type = EVisitOp::StructEpilogue });
+		}
+
+		bool ParseBody(FShallowTokenizer& Tk)
+		{
+			check(Tk.Is(ETokenType::BodyOpen));
+			Tk.Next();
+			
+			ParseBodyInner(Tk);
+			
+			return Tk.IsConsume(ETokenType::BodyClose);
+		}
+
+		void ParseBodyInner(FShallowTokenizer& Tk)
+		{
+			while (!Tk.IsEOS() && !Tk.Is(ETokenType::BodyClose))
+			{
+				ParseStatement(Tk);
+			}
+		}
+
+		void ParseStatement(FShallowTokenizer& Tk, bool bConsumeEOS = true)
+		{
+			OpBuffer.Add(FVisitOp { .Type = EVisitOp::StatementPrologue, .Statement = { .Cursor = Tk.Token.Begin } });
+			
+			if (!TryParseStatementInside(Tk))
+			{
+				OpBuffer.Add(FVisitOp { .Type = EVisitOp::StatementEpilogue, .Statement = { .Cursor = Tk.Begin } });
+				SkipStatement(Tk);
+				return;
+			}
+
+			OpBuffer.Add(FVisitOp { .Type = EVisitOp::StatementEpilogue, .Statement = { .Cursor = Tk.Token.Begin + Tk.Token.Length } });
+
+			if (bConsumeEOS)
+			{
+				Tk.IsConsume(ETokenType::EndOfStatement);
+			}
+		}
+
+		void ParseReturn(FShallowTokenizer& Tk)
+		{
+			int32 DrainAnchor = OpBuffer.Num();
+						
+			Tk.Next();
+
+			if (!Tk.Is(ETokenType::EndOfStatement))
+			{
+				TryParseExpression(Tk);
+
+				PushDrain(DrainAnchor, TEXT("Return"));
+			}
+		}
+
+		void ParseDXCPragma(FShallowTokenizer& Tk)
+		{
+			Tk.Next();
+			Tk.SkipUntil(')');
+			Tk.Next();
+		}
+
+		bool ParseIf(FShallowTokenizer& Tk)
+		{
+			bool bFirst = true;
+			
+			while (bFirst || Tk.IsConsume(ETokenType::Else))
+			{
+				bFirst = false;
+				
+				if (Tk.IsConsume(ETokenType::If))
+				{
+					if (!Tk.IsConsume(ETokenType::ParenthesisOpen))
+					{
+						return false;
+					}
+
+					WithStatement(Tk, [&] { TryParseExpression(Tk); });
+
+					if (!Tk.Is(ETokenType::ParenthesisClose))
+					{
+						return false;
+					}
+			
+					Tk.Next();
+				}
+
+				ParseStatementBody(Tk);
+			}
+
+			return true;
+		}
+
+		bool ParseWhile(FShallowTokenizer& Tk)
+		{
+			Tk.Next();
+			
+			if (!Tk.IsConsume(ETokenType::ParenthesisOpen))
+			{
+				return false;
+			}
+
+			WithStatement(Tk, [&] { TryParseExpression(Tk); });
+
+			if (!Tk.IsConsume(ETokenType::ParenthesisClose))
+			{
+				return false;
+			}
+
+			ParseStatementBody(Tk);
+			return true;
+		}
+
+		bool ParseFor(FShallowTokenizer& Tk)
+		{
+			Tk.Next();
+
+			if (!Tk.IsConsume(ETokenType::ParenthesisOpen))
+			{
+				return false;
+			}
+
+			// Initializer
+			if (!Tk.IsConsume(ETokenType::EndOfStatement))
+			{
+				ParseStatement(Tk, false);
+				
+				if (!Tk.IsConsume(ETokenType::EndOfStatement))
+				{
+					return false;
+				}
+			}
+
+			// Condition
+			if (!Tk.IsConsume(ETokenType::EndOfStatement))
+			{
+				ParseStatement(Tk, false);
+
+				if (!Tk.IsConsume(ETokenType::EndOfStatement))
+				{
+					return false;
+				}
+			}
+
+			// Step
+			if (!Tk.Is(ETokenType::ParenthesisClose))
+			{
+				ParseStatement(Tk);
+			}
+
+			if (!Tk.IsConsume(ETokenType::ParenthesisClose))
+			{
+				return false;
+			}
+
+			ParseStatementBody(Tk);
+
+			return true;
+		}
+
+		void ParseStatementBody(FShallowTokenizer& Tk)
+		{
+			if (Tk.Is(ETokenType::BodyOpen))
+			{
+				ParseBody(Tk);
+			}
+			else
+			{
+				ParseStatement(Tk);
+			}
+		}
+
+		bool ParseSwitch(FShallowTokenizer& Tk)
+		{
+			Tk.Next();
+
+			if (!Tk.IsConsume(ETokenType::ParenthesisOpen))
+			{
+				return false;
+			}
+
+			int32 DrainAnchor = OpBuffer.Num();
+						
+			TryParseExpression(Tk);
+			
+			PushDrain(DrainAnchor, TEXT("Switch-Cond"));
+
+			if(!Tk.IsConsume(ETokenType::ParenthesisClose))
+			{
+				return false;
+			}
+
+			if(!Tk.IsConsume(ETokenType::BodyOpen))
+			{
+				return false;
+			}
+
+			while (!Tk.IsEOS() && !Tk.Is(ETokenType::BodyClose))
+			{
+				if (Tk.IsConsume(ETokenType::Default))
+				{
+					
+				}
+				else if (Tk.IsConsume(ETokenType::Case))
+				{
+					if (!TryParsePrimary(Tk))
+					{
+						break;
+					}
+			
+					PushDrain(DrainAnchor, TEXT("Switch-Case"));
+				}
+				else
+				{
+					break;
+				}
+
+				if (!Tk.IsConsume(ETokenType::Colon))
+				{
+					break;
+				}
+
+				// Fallthrough?
+				if (Tk.Is(ETokenType::Case))
+				{
+					continue;
+				}
+
+				while (!Tk.Is(ETokenType::BodyClose) && !Tk.Is(ETokenType::Default) && !Tk.Is(ETokenType::Case))
+				{
+					ParseStatement(Tk);
+				}
+			}
+
+			return Tk.IsConsume(ETokenType::BodyClose);
+		}
+
+		void ParseBreak(FShallowTokenizer& Tk)
+		{
+			Tk.Next();
+		}
+
+		bool TryParseStatementInside(FShallowTokenizer& Tk)
+		{
+			// For now, just ignore all annotations
+			SkipAnnotations(Tk);
+
+			switch (Tk.Token.Type)
+			{
+				default:
+				{
+					break;
+				}
+				case ETokenType::Return:
+				{
+					ParseReturn(Tk);
+					return true;
+				}
+				case ETokenType::If:
+				{
+					return ParseIf(Tk);
+				}
+				case ETokenType::DXCPragma:
+				{
+					ParseDXCPragma(Tk);
+					return true;
+				}
+				case ETokenType::While:
+				{
+					return ParseWhile(Tk);
+				}
+				case ETokenType::Switch:
+				{
+					return ParseSwitch(Tk);
+				}
+				case ETokenType::For:
+				{
+					return ParseFor(Tk);
+				}
+				case ETokenType::Struct:
+				{
+					ParseStruct(Tk);
+					return true;
+				}
+				case ETokenType::BodyOpen:
+				{
+					return ParseBody(Tk);
+				}
+				case ETokenType::Break:
+				{
+					ParseBreak(Tk);
+					return true;
+				}
+			}
+			
+			if (TryParseFunction(Tk))
+			{
+				return true;
+			}
+			
+			if (TryParseDeclaration(Tk))
+			{
+				return true;
+			}
+			
+			if (TryParseExpression(Tk))
+			{
+				return true;
+			}
+
+			return false;
+		}
+
+		void SkipStatement(FShallowTokenizer& Tk)
+		{
+			if (!Tk.IsConsume(ETokenType::EndOfStatement))
+			{
+				// TODO: Structural skipping!
+				FParseCharType* Begin = Tk.Begin;
+				Tk.SkipUntil(';');
+
+				// Mark the full range as complex
+				MarkIdentifiersAsComplexBase(FParseViewType(Begin, Tk.Begin - Begin));
+
+	#if UE_SHADER_SDCE_LOG_COMPLEX
+				UE_LOG(LogTemp, Error, TEXT("Skipped complex statement: %hs"), *FParseStringType::ConstructFromPtrSize(Begin, Tk.Begin - Begin));
+	#endif // UE_SHADER_SDCE_LOG_COMPLEX
+			}
+		}
+
+		void MarkIdentifiersAsComplexBase(const FParseViewType& Region)
+		{
+			// Setup tokenizer
+			FShallowTokenizer Tk;
+			Tk.Begin = Region.Begin;
+			Tk.End = Region.Begin + Region.Length;
+			Tk.Next();
+
+			// Mark all identifiers as complex
+			while (!Tk.IsEOS())
+			{
+				if (Tk.Is(ETokenType::ID))
+				{
+					OpBuffer.Add(FVisitOp { .Type = EVisitOp::ComplexID, .ID = FSmallStringView::Get(Tk.Token) });
+				}
+				
+				Tk.Next();
+			}
+			
+			Tk.OpHead = TokenizerInvalidOpHead;
+		}
+		
+		bool TryParseFunction(FShallowTokenizer& TkOuter)
+		{
+			FShallowTokenizer Tk = Branch(TkOuter);
+
+			SkipTemplate(Tk);
+			SkipAnnotations(Tk);
+
+			FToken ReturnType;
+			if (!TryParseType(Tk, ReturnType))
+			{
+				Reject(Tk);
+				return false;
+			}
+
+			if (!TryParsePrimary(Tk, false))
+			{
+				Reject(Tk);
+				return false;
+			}
+
+			if (!Tk.IsConsume(ETokenType::ParenthesisOpen))
+			{
+				Reject(Tk);
+				return false;
+			}
+
+			OpBuffer.Add(FVisitOp { .Type = EVisitOp::FunctionPrologue });
+
+			// Parse all parameters
+			while (!Tk.IsEOS() && !Tk.Is(ETokenType::ParenthesisClose))
+			{
+				SkipQualifiers(Tk);
+
+				WithStatement(Tk, [&] { TryParseDeclaration(Tk, false); });
+
+				if (!Tk.IsConsume(ETokenType::Comma))
+				{
+					break;
+				}
+			}
+
+			if (!Tk.IsConsume(ETokenType::ParenthesisClose))
+			{
+				// TODO: Maybe this is less of a matching failure, and more of a chunk failure?
+				Reject(Tk);
+				return false;
+			}
+
+			if (Tk.Is(ETokenType::BodyOpen))
+			{
+				ParseBody(Tk);
+			}
+
+			OpBuffer.Add(FVisitOp { .Type = EVisitOp::FunctionEpilogue });
+
+			Accept(TkOuter, Tk);
+			return true;
+		}
+
+		bool TryParseDeclaration(FShallowTokenizer& TkOuter, bool bAllowMultiDeclarators = true)
+		{
+			FShallowTokenizer Tk = Branch(TkOuter);
+
+			SkipTemplate(Tk);
+			SkipQualifiers(Tk);
+
+			int32 TypeStart = OpBuffer.Num();
+			
+			FToken Type;
+			if (!TryParseType(Tk, Type))
+			{
+				Reject(Tk);
+				return false;
+			}
+
+			int32 TypeEnd = OpBuffer.Num();
+
+			// May have multiple declarators "int a, b, ..."
+			for (;;)
+			{
+				int32 DrainAnchor = OpBuffer.Num();
+
+				// Push type ops again if secondary
+				if (OpBuffer.Num() != TypeEnd)
+				{
+					for (int i = TypeStart; i < TypeEnd; ++i)
+					{
+						OpBuffer.Add(OpBuffer[i]);
+					}
+				}
+				
+				FToken ID = Tk.Next();
+				if (ID.Type != ETokenType::ID)
+				{
+					Reject(Tk);
+					return false;
+				}
+
+				OpBuffer.Add(FVisitOp { .Type = EVisitOp::ID, .ID = FSmallStringView::Get(ID) });
+
+				// Declarators may have type postfixes after the ID
+				TryParseTypePostfix(Tk);
+
+				// Don't care about bindings
+				SkipDeclarationBindings(Tk);
+
+				OpBuffer.Add(FVisitOp { .Type = EVisitOp::Declaration });
+
+				// Handle inline assignments
+				if (Tk.Is(ETokenType::BinaryEq))
+				{
+					OpBuffer.Add(FVisitOp { .Type = EVisitOp::Binary, .Binary = { Tk.Token } });
+
+					Tk.Next();
+				
+					if (!TryParseExpression(Tk, false))
+					{
+						Reject(Tk);
+						return false;
+					}
+				}
+
+				if (!bAllowMultiDeclarators || !Tk.IsConsume(ETokenType::Comma))
+				{
+					break;
+				}
+
+				// Last value carries over
+				PushDrain(DrainAnchor, TEXT("MultiDeclarator"));
+			}
+
+			Accept(TkOuter, Tk);
+			return true;
+		}
+
+		bool TryParseType(FShallowTokenizer& Tk, FToken& Type)
+		{
+			// Note: This is obviously an over simplified type
+			// When needed, it'll be expanded on
+			
+			Type = Tk.Next();
+			if (Type.Type != ETokenType::ID)
+			{
+				return false;
+			}
+
+			OpBuffer.Add(FVisitOp { .Type = EVisitOp::Type, .ID = { FSmallStringView::Get(Type) } });
+
+			TryParseTypePostfix(Tk);
+
+			if (TemplatedTypes.Contains(FSmallStringView::Get(Type)))
+			{
+				SkipTemplateSpecialization(Tk);
+			}
+
+			return true;
+		}
+
+		void TryParseTypePostfix(FShallowTokenizer& Tk)
+		{
+			for (;;)
+			{
+				if (Tk.IsConsume(ETokenType::SquareOpen))
+				{
+					// May close immediately for automatic sizing
+					if (!Tk.Is(ETokenType::SquareClose))
+					{
+						int32 DrainAnchor = OpBuffer.Num();
+						TryParseExpression(Tk);
+						PushDrain(DrainAnchor, TEXT("TypeArrayDim"));
+					}
+					
+					check(Tk.Is(ETokenType::SquareClose));
+					Tk.Next();
+					continue;
+				}
+
+				// Not a postfix
+				break;
+			}
+		}
+
+		bool IsUnaryPrimaryValueModifier(ETokenType Type)
+		{
+			switch (Type)
+			{
+			default:
+				return false;
+				case ETokenType::BinaryAdd:
+				case ETokenType::BinarySub:
+				case ETokenType::UnaryBitNegate:
+				case ETokenType::UnaryNot:
+				case ETokenType::UnaryInc:
+				case ETokenType::UnaryDec:
+					return true;
+			}
+		}
+
+		bool ParseParenthesis(FShallowTokenizer& Tk)
+		{
+			Tk.Next();
+
+			OpBuffer.Add(FVisitOp { .Type = EVisitOp::FramePrologue, .bPrivate = false });
+			
+			int32 DrainAnchor = OpBuffer.Num();
+			TryParseExpression(Tk);
+
+			// Expecting close
+			if (!Tk.IsConsume(ETokenType::ParenthesisClose))
+			{
+				return false;
+			}
+
+			// Cast?
+			if (Tk.Is(ETokenType::ID) || Tk.Is(ETokenType::Numeric) || IsUnaryPrimaryValueModifier(Tk.Token.Type))
+			{
+				PushDrain(DrainAnchor, TEXT("Cast"));
+
+				OpBuffer.Add(FVisitOp { .Type = EVisitOp::FrameEpilogue, .bPrivate = false });
+				
+				return TryParsePrimary(Tk);
+			}
+			
+			OpBuffer.Add(FVisitOp { .Type = EVisitOp::FrameEpilogue, .bPrivate = false });
+
+			// Otherwise, allow accessors
+			TryParsePrimaryAccessors(Tk);
+
+			return true;
+		}
+
+		void ParseTernary(FShallowTokenizer& Tk, bool bAllowExpressionCarry = true)
+		{
+			check(Tk.Is(ETokenType::Ternary));
+			Tk.Next();
+
+			// Pass value
+			{
+				OpBuffer.Add(FVisitOp { .Type = EVisitOp::FramePrologue, .bPrivate = false });
+			
+				int32 DrainAnchor = OpBuffer.Num();
+				TryParseExpression(Tk, bAllowExpressionCarry);
+				PushDrain(DrainAnchor, TEXT("TenaryPass"));
+			
+				OpBuffer.Add(FVisitOp { .Type = EVisitOp::FrameEpilogue, .bPrivate = false });
+			}
+
+			check(Tk.Is(ETokenType::Colon));
+			Tk.Next();
+
+			// Failure value
+			{
+				OpBuffer.Add(FVisitOp { .Type = EVisitOp::FramePrologue, .bPrivate = false });
+			
+				int32 DrainAnchor = OpBuffer.Num();
+				TryParseExpression(Tk, bAllowExpressionCarry);
+				PushDrain(DrainAnchor, TEXT("TernaryFail"));
+			
+				OpBuffer.Add(FVisitOp { .Type = EVisitOp::FrameEpilogue, .bPrivate = false });
+			}
+		}
+
+		bool TryParseExpression(FShallowTokenizer& TkOuter, bool bAllowExpressionCarry = true)
+		{
+			FShallowTokenizer Tk = Branch(TkOuter);
+
+			int32 DrainAnchor = OpBuffer.Num();
+
+			if (!TryParsePrimary(Tk))
+			{
+				Reject(Tk);
+				return false;
+			}
+
+			if (!TryParseBinary(Tk))
+			{
+				Reject(Tk);
+				return false;
+			}
+
+			if (Tk.Is(ETokenType::Ternary))
+			{
+				ParseTernary(Tk, bAllowExpressionCarry);
+			}
+
+			if (bAllowExpressionCarry && Tk.IsConsume(ETokenType::Comma))
+			{
+				PushDrain(DrainAnchor, TEXT("CommaCarry"));
+
+				if (!TryParseExpression(Tk))
+				{
+					Reject(Tk);
+					return false;
+				}
+			}
+			
+			Accept(TkOuter, Tk);
+			return true;
+		}
+
+		bool ParseUnary(FShallowTokenizer& Tk)
+		{
+			Tk.Next();
+			return TryParsePrimary(Tk);
+		}
+
+		bool ParseAggregateInitializer(FShallowTokenizer& Tk)
+		{
+			Tk.Next();
+
+			while (!Tk.IsEOS() && !Tk.Is(ETokenType::BodyClose))
+			{
+				int32 DrainAnchor = OpBuffer.Num();
+				
+				if (!TryParseExpression(Tk, false))
+				{
+					return false;
+				}
+		
+				PushDrain(DrainAnchor, TEXT("AggregateInitMember"));
+
+				if (!Tk.IsConsume(ETokenType::Comma))
+				{
+					break;
+				}
+			}
+
+			// Zero for now, until we need it
+			OpBuffer.Add(FVisitOp { .Type = EVisitOp::AggregateComposite });
+					
+			return Tk.IsConsume(ETokenType::BodyClose);
+		}
+
+		bool TryParsePrimary(FShallowTokenizer& TkOuter, bool bAllowCall = true)
+		{
+			if (TkOuter.Is(ETokenType::ParenthesisOpen))
+			{
+				return ParseParenthesis(TkOuter);
+			}
+
+			if (IsUnaryPrimaryValueModifier(TkOuter.Token.Type))
+			{
+				return ParseUnary(TkOuter);
+			}
+
+			if (TkOuter.Is(ETokenType::BodyOpen))
+			{
+				return ParseAggregateInitializer(TkOuter);
+			}
+
+			FShallowTokenizer Tk = Branch(TkOuter);
+			
+			FToken PrimID = Tk.Next();
+			if (PrimID.Type != ETokenType::ID && PrimID.Type != ETokenType::Numeric)
+			{
+				Reject(Tk);
+				return false;
+			}
+
+			OpBuffer.Add(FVisitOp { .Type = EVisitOp::ID, .ID = { FSmallStringView::Get(PrimID) } });
+			
+			if (!TryParsePrimaryAccessors(Tk, bAllowCall))
+			{
+				Reject(Tk);
+				return false;
+			}
+
+			Accept(TkOuter, Tk);
+			return true;
+		}
+
+		void PushDrain(int32 Anchor, const TCHAR* Reason)
+		{
+			// TODO: Ignore drains on statement ends
+			EPrecedence Precedence = EPrecedence::None;
+
+			// Zero op range? early out
+			if (Anchor == OpBuffer.Num()) 
+			{
+				return;
+			}
+
+			// Current number of private frames
+			int32 PrivateFrameCounter = 0;
+
+			// We need to maintain the precedence of the value that we're draining,
+			// otherwise we violate the op "hierarchy" (it's flat, but the implicit one).
+			for (int32 i = Anchor; i < OpBuffer.Num(); i++)
+			{
+				EPrecedence VisitPrecedence;
+
+				// Get precedence for the op
+				const FVisitOp& Op = OpBuffer[i];
+				switch (Op.Type)
+				{
+					default:
+					{
+						checkNoEntry();
+						continue;
+					}
+					case EVisitOp::ID:
+					{
+						VisitPrecedence = EPrecedence::Primary;
+						break;
+					}
+					case EVisitOp::Type:
+					{
+						VisitPrecedence = EPrecedence::Primary;
+						break;
+					}
+					case EVisitOp::Binary:
+					{
+						VisitPrecedence = Op.Binary.Token.Type == ETokenType::BinaryEq ? EPrecedence::BinaryAssign : EPrecedence::BinaryGeneric;
+						break;
+					}
+					case EVisitOp::MemberAccess: {
+							
+						VisitPrecedence = EPrecedence::Access;
+						break;
+					}
+					case EVisitOp::Declaration:
+					{
+						VisitPrecedence = EPrecedence::Declare;
+						break;
+					}
+					case EVisitOp::AggregateComposite:
+					{
+						VisitPrecedence = EPrecedence::Primary;
+						break;
+					}
+					case EVisitOp::FramePrologue:
+					{
+						if (Op.bPrivate)
+						{
+							PrivateFrameCounter++;
+						}
+						continue;
+					}
+					case EVisitOp::FrameEpilogue:
+					{
+						if (Op.bPrivate)
+						{
+							PrivateFrameCounter--;
+						}
+						continue;
+					}
+					case EVisitOp::StatementPrologue:
+					case EVisitOp::StatementEpilogue:
+					case EVisitOp::StructPrologue:
+					case EVisitOp::StructEpilogue:
+					case EVisitOp::FunctionPrologue:
+					case EVisitOp::FunctionEpilogue:
+					case EVisitOp::Alias: 
+					case EVisitOp::Drain:
+					{
+						VisitPrecedence = EPrecedence::None;
+						break;
+					}
+				}
+
+				// If we're in a private frame, do not account for it
+				if (PrivateFrameCounter > 0)
+				{
+					continue;
+				}
+
+				Precedence = FMath::Max(Precedence, VisitPrecedence);
+			}
+
+			// Must close out
+			check(PrivateFrameCounter == 0);
+
+			// Must have found at least one op
+			check(Precedence != EPrecedence::None);
+
+			OpBuffer.Add(FVisitOp { .Type = EVisitOp::Drain, .Drain = { Reason, Precedence } });
+		}
+
+		bool TryParsePrimaryAccessors(FShallowTokenizer& Tk, bool bAllowCall = true)
+		{
+			for (;;)
+			{
+				// ++/--
+				if (Tk.Is(ETokenType::UnaryInc) || Tk.Is(ETokenType::UnaryDec))
+				{
+					Tk.Next();
+					continue;
+				}
+
+				// Array accessor
+				if (Tk.Is(ETokenType::SquareOpen))
+				{
+					Tk.Next();
+					
+					int32 DrainAnchor = OpBuffer.Num();
+						
+					if (!TryParseExpression(Tk))
+					{
+						return false;
+					}
+
+					PushDrain(DrainAnchor, TEXT("ArrayAccessor"));
+
+					check(Tk.Is(ETokenType::SquareClose));
+					Tk.Next();
+					
+					continue;
+				}
+
+				// Member access
+				if (Tk.Is(ETokenType::Access) || Tk.Is(ETokenType::NamespaceAccess))
+				{
+					Tk.Next();
+					
+					FToken ID = Tk.Next();
+					if (ID.Type != ETokenType::ID)
+					{
+						return false;
+					}
+
+					OpBuffer.Add(FVisitOp { .Type = EVisitOp::MemberAccess, .ID = { FSmallStringView::Get(ID) } });
+					
+					continue;
+				}
+				
+				// Calls
+				if (bAllowCall && Tk.Is(ETokenType::ParenthesisOpen))
+				{
+					Tk.Next();
+
+					while (!Tk.IsEOS() && !Tk.Is(ETokenType::ParenthesisClose))
+					{
+						OpBuffer.Add(FVisitOp { .Type = EVisitOp::FramePrologue, .bPrivate = true });
+					
+						int32 DrainAnchor = OpBuffer.Num();
+						
+						if (!TryParseExpression(Tk, false))
+						{
+							return false;
+						}
+
+						PushDrain(DrainAnchor, TEXT("CallArg"));
+						
+						OpBuffer.Add(FVisitOp { .Type = EVisitOp::FrameEpilogue, .bPrivate = true });
+
+						if (!Tk.IsConsume(ETokenType::Comma))
+						{
+							break;
+						}
+					}
+
+					// Unexpected closure
+					if (!Tk.Is(ETokenType::ParenthesisClose))
+					{
+						return false;
+					}
+					
+					Tk.Next();
+					continue;
+				}
+
+				// Not an accessor
+				break;
+			}
+
+			// OK
+			return true;
+		}
+
+		bool TryParseBinary(FShallowTokenizer& Tk)
+		{
+			// Note: Other than assignments, we are effectively not caring about precedence
+			
+			while (!Tk.IsEOS())
+			{
+				// Is this a binary operation?
+				switch (Tk.Token.Type)
+				{
+				default:
+					// OK
+					return true;
+				case ETokenType::BinaryEq:
+				case ETokenType::BinaryNotEq:
+				case ETokenType::BinaryPlusEq:
+				case ETokenType::BinarySubEq:
+				case ETokenType::BinaryDivEq:
+				case ETokenType::BinaryMulEq:
+				case ETokenType::BinaryModEq:
+				case ETokenType::BinaryXEq:
+				case ETokenType::BinaryOrEq:
+				case ETokenType::BinaryAndEq:
+				case ETokenType::BinaryShlEq:
+				case ETokenType::BinaryShrEq:
+				case ETokenType::LogicalAnd:
+				case ETokenType::LogicalOr:
+				case ETokenType::BinaryLess:
+				case ETokenType::BinaryLessEq:
+				case ETokenType::BinaryGreater:
+				case ETokenType::BinaryGreaterEq:
+				case ETokenType::BinaryLogicalEq:
+				case ETokenType::BinaryAdd:
+				case ETokenType::BinarySub:
+				case ETokenType::BinaryDiv:
+				case ETokenType::BinaryMul:
+				case ETokenType::BinaryMod:
+				case ETokenType::BinaryXor:
+				case ETokenType::BinaryBitOr:
+				case ETokenType::BinaryBitAnd:
+				case ETokenType::BinaryShl:
+				case ETokenType::BinaryShr:
+					break;
+				}
+
+				OpBuffer.Add(FVisitOp { .Type = EVisitOp::Binary, .Binary = { Tk.Token } });
+				
+				Tk.Next();
+
+				if (!TryParsePrimary(Tk))
+				{
+					return false;
+				}
+			}
+
+			return true;
+		}
+
+		bool TryParseUsingAlias(FShallowTokenizer& Tk)
+		{
+			check(Tk.Is(ETokenType::Using));
+			Tk.Next();
+
+			FToken ID = Tk.Next();
+			if (ID.Type != ETokenType::ID)
+			{
+				return false;
+			}
+
+			if (!Tk.IsConsume(ETokenType::BinaryEq))
+			{
+				return false;
+			}
+			
+			FToken Type;
+			if (!TryParseType(Tk, Type))
+			{
+				return false;
+			}
+
+			OpBuffer.Add(FVisitOp { .Type = EVisitOp::Alias });
+
+			return true;
+		}
+
+		bool TryParseTypedefAlias(FShallowTokenizer& Tk)
+		{
+			check(Tk.Is(ETokenType::Typedef));
+			Tk.Next();
+
+			FToken Type;
+			if (!TryParseType(Tk, Type))
+			{
+				return false;
+			}
+			
+			FToken ID = Tk.Next();
+			if (ID.Type != ETokenType::ID)
+			{
+				return false;
+			}
+
+			OpBuffer.Add(FVisitOp { .Type = EVisitOp::Alias });
+
+			return true;
+		}
+
+		void SkipAnnotations(FShallowTokenizer& Tk)
+		{
+			while (Tk.Is(ETokenType::SquareOpen))
+			{
+				// TODO: Structural skipping!
+				Tk.SkipUntil(TEXT(']'));
+				Tk.Next();
+			}
+		}
+
+		bool SkipTemplate(FShallowTokenizer& Tk)
+		{
+			if (!Tk.IsConsume(ETokenType::Template))
+			{
+				return false;
+			}
+
+			if (!Tk.Is(ETokenType::BinaryGreater))
+			{
+				// TODO: Structural skipping!
+				Tk.SkipUntil(TEXT('>'));
+			}
+			
+			Tk.Next();
+			return true;
+		}
+
+		void SkipTemplateSpecialization(FShallowTokenizer& Tk)
+		{
+			if (!Tk.IsConsume(ETokenType::BinaryLess))
+			{
+				return;
+			}
+
+			if (!Tk.Is(ETokenType::BinaryGreater))
+			{
+				// TODO: Structural skipping!
+				Tk.SkipUntil(TEXT('>'));
+			}
+			
+			Tk.Next();
+		}
+
+		void SkipDeclarationBindings(FShallowTokenizer& Tk)
+		{
+			if (!Tk.IsConsume(ETokenType::Colon))
+			{
+				return;
+			}
+
+			// If this is a binding, make an educated guess
+			if (!Tk.Is(ETokenType::Comma) && !Tk.Is(ETokenType::EndOfStatement) && !Tk.Is(ETokenType::ParenthesisClose))
+			{
+				Tk.SkipUntil(TEXT(",;"));
+			}
+		}
+
+		void SkipQualifiers(FShallowTokenizer& Tk)
+		{
+			for (;;)
+			{
+				switch (Tk.Token.Type)
+				{
+				default:
+					return;
+				case ETokenType::Static:
+				case ETokenType::Const:
+				case ETokenType::In:
+				case ETokenType::Out:
+				case ETokenType::Inout:
+				case ETokenType::NoInterp:
+				case ETokenType::NoInterpolation:
+				case ETokenType::NoPerspective:
+				case ETokenType::GroupShared:
+				case ETokenType::Triangle:
+				case ETokenType::Centroid:
+				case ETokenType::Vertices:
+				case ETokenType::Indices:
+				case ETokenType::Primitives:
+				case ETokenType::Uniform:
+					Tk.Next();
+					break;
+				}
+			}
+		}
+
+	private:
+		FShallowTokenizer Branch(const FShallowTokenizer& Tk)
+		{
+			FShallowTokenizer Branch = Tk.Branch();
+			Branch.OpHead = OpBuffer.Num();
+			return MoveTemp(Branch);
+		}
+
+		void Reject(FShallowTokenizer& Tk)
+		{
+			OpBuffer.SetNumNoAlloc(Tk.OpHead);
+			Tk.OpHead = TokenizerInvalidOpHead;
+		}
+
+		void Accept(FShallowTokenizer& TkOuter, FShallowTokenizer& Tk)
+		{
+			TkOuter.Accept(Tk);
+		}
+
+	private:
+		enum class EVisitOp
+		{
+			ID,
+			ComplexID,
+			Type,
+			Binary,
+			MemberAccess,
+			AggregateComposite,
+			Declaration,
+			StatementPrologue,
+			StatementEpilogue,
+			StructPrologue,
+			StructEpilogue,
+			FunctionPrologue,
+			FunctionEpilogue,
+			FramePrologue,
+			FrameEpilogue,
+			Alias,
+			Drain
+		};
+
+		struct FVisitOp
+		{
+			EVisitOp Type;
+
+			/** Payload */
+			union
+			{
+				/** EVisitOp::ID/ComplexID */
+				FSmallStringView ID;
+
+				/** EVisitOp::FramePrologue/FrameEpilogue */
+				bool bPrivate;
+				
+				/** EVisitOp::StructPrologue */
+				struct
+				{
+					FSmallStringView ID;
+				} Struct;
+				
+				/** EVisitOp::Binary */
+				struct
+				{
+					FToken Token;
+				} Binary;
+
+				/** EVisitOp::Drain */
+				struct
+				{
+					const TCHAR* Reason;
+					EPrecedence Type;
+				} Drain;
+				
+				/** EVisitOp::StatementPrologue/StatementEpilogue */
+				struct
+				{
+					FParseCharType* Cursor;
+				} Statement;
+			};
+		};
+
+		static_assert(sizeof(FVisitOp) == 24, "Unexpected FVisitOp size");
+
+		/** Current visitation buffer */
+		TOpBuffer<FVisitOp, 256> OpBuffer;
+
+	private:
+		/** Tracked templated types */
+		TSet<FConstSmallStringView, FConstSmallStringKeyFuncs, FMemStackSetAllocator> TemplatedTypes;
+	};
+}
diff --git a/Engine/Source/Developer/ShaderCompilerCommon/Private/SDCE/SDCESymbolContext.h b/Engine/Source/Developer/ShaderCompilerCommon/Private/SDCE/SDCESymbolContext.h
new file mode 100644
--- /dev/null
+++ b/Engine/Source/Developer/ShaderCompilerCommon/Private/SDCE/SDCESymbolContext.h
@@ -0,0 +1,213 @@
+﻿// Copyright Epic Games, Inc. All Rights Reserved.
+#pragma once
+
+#include "SDCECommon.h"
+#include "SDCEToken.h"
+#include "Containers/ChunkedArray.h"
+
+namespace UE::ShaderMinifier::SDCE
+{
+	enum class ESymbolKind
+	{
+		/** Generic composite type, e.g., structs */
+		Composite,
+
+		/** Declaration symbols */
+		Variable
+	};
+		
+	struct FSymbol
+	{
+		/** Kind of this symbol */
+		ESymbolKind Kind;
+
+		/** Name of this symbol */
+		FSmallStringView ID{};
+
+		/** Is this symbol tagged? i.e., can we perform SDCE on it? */
+		bool bTagged = false;
+	
+		struct
+		{
+			/** Type of this variable */
+			FSymbol* Type = nullptr;
+
+			/** Number of realized loads, zero implies safe DCE'ing */
+			int32 RealizedLoadCount = 0;
+		} Variable;
+	
+		struct
+		{
+			/**
+			 * All symbols within
+			 * TODO: I really, really don't like bloating all symbols like this
+			 */
+			TArray<FSymbol*, TMemStackAllocator<>> ContainedSymbols;
+		} Composite;
+	};
+	
+	struct FSymbolScope
+	{
+		/** Owning composite of this scope */
+		FSymbol* Composite = nullptr;
+
+		/** All declared symbols within this scope */
+		TMap<FSmallStringView, FSymbol*,  FMemStackSetAllocator, FSmallStringKeyFuncs> Symbols;
+	};
+	
+	/** Value span, represents a range of values */
+	struct FValueSpan
+	{
+		FValueSpan() : FValueSpan(0, 0)
+		{
+			
+		}
+	
+		FValueSpan(int32 InBegin, int32 InEnd)
+		{
+			Begin = InBegin;
+			End = InEnd;
+		}
+				
+		FValueSpan operator|(const FValueSpan& RHS) const
+		{
+			FValueSpan Span;
+			Span.Begin = FMath::Min(Begin, RHS.Begin);
+			Span.End = FMath::Max(End, RHS.End);
+			return Span;
+		}
+				
+		FValueSpan& operator|=(const FValueSpan& RHS)
+		{
+			Begin = FMath::Min(Begin, RHS.Begin);
+			End = FMath::Max(End, RHS.End);
+			return *this;
+		}
+				
+		int32 Begin;
+		int32 End;
+	};
+
+	struct FMemberStoreSite
+	{
+		/** Symbols */
+		FSymbol* AddressSymbol = nullptr;
+		FSymbol* ValueSymbol   = nullptr;
+
+		/** Value spans */
+		FValueSpan AddressSpan;
+		FValueSpan ValueSpan;
+
+		/** Has this site been visited? Used for iterative elimination */
+		bool bVisited = false;
+
+		/** Was this determined to be a complex site? */
+		bool bIsResolvedToComplex = false;
+
+		/** Textual site */
+		FSmallStringView AssignmentSite;
+	};
+	
+	class FSymbolContext
+	{
+	public:
+		FSymbolContext()
+		{
+			// Root stack
+			SymbolStack.Emplace();
+		}
+
+		/** Find a symbol within a composite */
+		FSymbol* LinearCompositeFindSymbol(FSymbol* Base, const FSmallStringView& View)
+		{
+			if (Base->Kind != ESymbolKind::Composite)
+			{
+				return nullptr;
+			}
+			
+			for (FSymbol* CandidateSymbol : Base->Composite.ContainedSymbols)
+			{
+				if (CandidateSymbol->ID == View)
+				{
+					return CandidateSymbol;
+				}
+			}
+	
+			return nullptr;
+		}
+
+		/** Find a symbol in the current scope */
+		FSymbol* FindSymbol(const FSmallStringView& View)
+		{
+			if (!View.Begin)
+			{
+				return nullptr;
+			}
+			
+			// Allow resolving to the first parent aggregate symbol scope
+			// Anything beyond that is ill-defined
+			bool bIsImmediateFirstAggregate = true;
+			
+			// Try inner to outer
+			for (int32 i = SymbolStack.Num() - 1; i >= 0; i--)
+			{
+				FSymbolScope& Scope = SymbolStack[i];
+				
+				if (Scope.Composite)
+				{
+					if (!bIsImmediateFirstAggregate)
+					{
+						// This is a private scope
+						continue;
+					}
+	
+					// We are now in "this"
+					bIsImmediateFirstAggregate = false;
+				}
+				
+				if (FSymbol** It = Scope.Symbols.Find(View))
+				{
+					return *It;
+				}
+			}
+	
+			return nullptr;
+		}
+
+		/** Create a new symbol */
+		FSymbol* CreateSymbol(ESymbolKind Kind)
+		{
+			FSymbol* Symbol = &Symbols[Symbols.Emplace()];
+			Symbol->Kind = Kind;
+			return Symbol;
+		}
+
+		/** Add a complex accessor to the set */
+		void AddComplexAccessor(FSmallStringView ID)
+		{
+#if UE_SHADER_SDCE_LOG_COMPLEX
+			UE_LOG(LogTemp, Error, TEXT("Added complex symbol: %hs"), *FAnsiString::ConstructFromPtrSize(ID.Begin, ID.Length));
+#endif // UE_SHADER_SDCE_LOG_COMPLEX
+	
+			if (ID.Begin)
+			{
+				ComplexBaseSet.Add(ID);
+			}
+		}
+
+		/** All symbol stacks */
+		TArray<FSymbolScope, TMemStackAllocator<>> SymbolStack;
+
+		/** All accessors that we're treating as complex */
+		TSet<FSmallStringView, FSmallStringKeyFuncs, FMemStackSetAllocator> ComplexBaseSet;
+
+		/** All tracked store sites */
+		TArray<FMemberStoreSite, TMemStackAllocator<>> StoreSites;
+		
+	private:
+		/** Symbol container */
+		TChunkedArray<FSymbol, 256, TMemStackAllocator<>> Symbols;
+	};
+
+	static_assert(sizeof(FSymbol) == 64, "Unexpected FSymbol packing");
+}
diff --git a/Engine/Source/Developer/ShaderCompilerCommon/Private/SDCE/SDCEToken.h b/Engine/Source/Developer/ShaderCompilerCommon/Private/SDCE/SDCEToken.h
new file mode 100644
--- /dev/null
+++ b/Engine/Source/Developer/ShaderCompilerCommon/Private/SDCE/SDCEToken.h
@@ -0,0 +1,234 @@
+﻿// Copyright Epic Games, Inc. All Rights Reserved.
+#pragma once
+
+#include "SDCECommon.h"
+#include "Containers/SetUtilities.h"
+#include "Hash/xxhash.h"
+
+namespace UE::ShaderMinifier::SDCE
+{
+	/** Full keyword/segment table */
+	#define UE_MINIFIER_SDCE_KEYWORD_TABLE \
+		X(If, "if") \
+		X(Else, "else") \
+		X(Switch, "switch") \
+		X(Case, "case") \
+		X(Default, "default") \
+		X(For, "for") \
+		X(While, "while") \
+		X(Do, "do") \
+		X(Break, "break") \
+		X(Continue, "continue") \
+		X(Discard, "discard") \
+		X(Return, "return") \
+		X(Struct, "struct")  \
+		X(Namespace, "namespace") \
+		X(BinaryEq, "=") \
+		X(BinaryNotEq, "!=") \
+		X(BinaryPlusEq, "+=") \
+		X(BinarySubEq, "-=") \
+		X(BinaryDivEq, "/=") \
+		X(BinaryMulEq, "*=") \
+		X(BinaryModEq, "%=") \
+		X(BinaryXEq, "^=") \
+		X(BinaryOrEq, "|=") \
+		X(BinaryAndEq, "&=") \
+		X(BinaryShlEq, "<<=") \
+		X(BinaryShrEq, ">>=") \
+		X(LogicalAnd, "&&") \
+		X(LogicalOr, "||") \
+		X(BinaryLess, "<") \
+		X(BinaryLessEq, "<=") \
+		X(BinaryGreater, ">") \
+		X(BinaryGreaterEq, ">=") \
+		X(BinaryLogicalEq, "==") \
+		X(BinaryAdd, "+") \
+		X(BinarySub, "-") \
+		X(BinaryDiv, "/") \
+		X(BinaryMul, "*") \
+		X(BinaryMod, "%") \
+		X(BinaryXor, "^") \
+		X(BinaryBitOr, "|") \
+		X(BinaryBitAnd, "&") \
+		X(BinaryShl, "<<") \
+		X(BinaryShr, ">>") \
+		X(BodyOpen, "{") \
+		X(BodyClose, "}") \
+		X(ParenthesisOpen, "(") \
+		X(ParenthesisClose, ")") \
+		X(SquareOpen, "[") \
+		X(SquareClose, "]") \
+		X(EndOfStatement, ";") \
+		X(Ternary, "?") \
+		X(Comma, ",") \
+		X(Colon, ":") \
+		X(Access, ".") \
+		X(NamespaceAccess, "::") \
+		X(UnaryBitNegate, "~") \
+		X(UnaryNot, "!") \
+		X(UnaryInc, "++") \
+		X(UnaryDec, "--") \
+		X(Using, "using") \
+		X(Typedef, "typedef") \
+		X(Template, "template") \
+		X(Static, "static") \
+		X(Const, "const") \
+		X(In, "in") \
+		X(Out, "out") \
+		X(Inout, "inout") \
+		X(NoInterp, "nointerp") \
+		X(NoInterpolation, "nointerpolation") \
+		X(NoPerspective, "noperspective") \
+		X(GroupShared, "groupshared") \
+		X(Triangle, "triangle") \
+		X(Vertices, "vertices") \
+		X(Indices, "indices") \
+		X(Primitives, "primitives") \
+		X(Uniform, "uniform") \
+		X(Centroid, "centroid") \
+		X(DXCPragma, "_Pragma")
+
+	enum class ETokenType : uint16_t
+	{
+		/** Standard identifier */
+		ID,
+
+		/** Numeric identifier */
+		Numeric,
+
+		/** Unknown token */
+		Unknown,
+
+		/** Segments */
+#define X(N, S) N,
+		UE_MINIFIER_SDCE_KEYWORD_TABLE
+#undef X
+	};
+		
+	struct FToken
+	{
+		/** Set from an identifier */
+		void SetIDNoHash(const FParseViewType& View)
+		{
+			Begin = View.Begin;
+			Length = static_cast<uint16_t>(View.Length);
+			XXHash32 = 0;
+		}
+
+		// DW0
+		ETokenType Type;
+		uint16_t Length;
+		// DW1
+		// Each token inlines the identifier hash, this is faster than doing it on demand
+		uint32_t XXHash32;
+		// DW2-4
+		FParseCharType* Begin;
+	};
+	
+	struct FSmallStringView
+	{
+		/** Get from a token  */
+		static FSmallStringView Get(const FToken& Token)
+		{
+			return FSmallStringView {
+				Token.Begin,
+				Token.XXHash32,
+				Token.Length
+			};
+		}
+
+		/** Hash a text region */
+		static uint32_t Hash(const FParseCharType* Data, uint16 Length)
+		{
+			// TODO: Mod?
+			// TODO: Hash only the ansi bits?
+			return static_cast<uint32_t>(FXxHash64::HashBuffer(Data, Length * sizeof(FParseCharType)).Hash);
+		}
+
+		/** Comparator */
+		bool operator==(const FSmallStringView& RHS) const {
+			return
+				XXHash32 == RHS.XXHash32 &&
+				Length == RHS.Length &&
+				!FMemory::Memcmp(Begin, RHS.Begin, sizeof(FParseCharType) * Length);
+		}
+
+		// DW0-2
+		FParseCharType* Begin;
+		// DW2
+		uint32_t XXHash32;
+		uint16_t Length;
+	};
+	
+	struct FConstSmallStringView
+	{
+		FConstSmallStringView() = default;
+		
+		FConstSmallStringView(const FSmallStringView& Mutable) : Begin(Mutable.Begin), XXHash32(Mutable.XXHash32), Length(Mutable.Length)
+		{
+			
+		}
+		
+		/** Create from a text region */
+		static FConstSmallStringView Text(const FParseCharType* Data)
+		{
+			return Text(Data, FCString::Strlen(Data));
+		}
+		
+		/** Create from a text region */
+		static FConstSmallStringView Text(const FParseCharType* Data, uint16 Length)
+		{
+			FConstSmallStringView Out;
+			Out.Begin = Data;
+			Out.XXHash32 = FSmallStringView::Hash(Data, Length);
+			Out.Length = Length;
+			return Out;
+		}
+
+		/** Comparator */
+		bool operator==(const FConstSmallStringView& RHS) const {
+			return
+				XXHash32 == RHS.XXHash32 &&
+				Length == RHS.Length &&
+				!FMemory::Memcmp(Begin, RHS.Begin, sizeof(FParseCharType) * Length);
+		}
+		
+		// DW0-2
+		const FParseCharType* Begin;
+		// DW2
+		uint32_t XXHash32;
+		uint16_t Length;
+	};
+
+	template<typename T>
+	struct TSmallStringKeyFuncs : DefaultKeyFuncs<T>
+	{
+		static FORCEINLINE T GetSetKey(T K)
+		{
+			return K;
+		}
+
+		template <typename U>
+		static FORCEINLINE T GetSetKey(const TPair<T, U>& P)
+		{
+			return P.Key;
+		}
+		
+		static FORCEINLINE bool Matches(T A, T B)
+		{
+			return A == B;
+		}
+		
+		static FORCEINLINE uint32 GetKeyHash(T Key)
+		{
+			return Key.XXHash32;
+		}
+	};
+
+	using FSmallStringKeyFuncs      = TSmallStringKeyFuncs<FSmallStringView>;
+	using FConstSmallStringKeyFuncs = TSmallStringKeyFuncs<FConstSmallStringView>;
+
+	static_assert(sizeof(FToken) == 16,                "Unexpected FToken packing");
+	static_assert(sizeof(FSmallStringView) == 16,      "Unexpected FSmallStringView packing");
+	static_assert(sizeof(FConstSmallStringView) == 16, "Unexpected FConstSmallStringView packing");
+}
diff --git a/Engine/Source/Developer/ShaderCompilerCommon/Private/SDCE/SDCETokenizer.h b/Engine/Source/Developer/ShaderCompilerCommon/Private/SDCE/SDCETokenizer.h
new file mode 100644
--- /dev/null
+++ b/Engine/Source/Developer/ShaderCompilerCommon/Private/SDCE/SDCETokenizer.h
@@ -0,0 +1,459 @@
+﻿// Copyright Epic Games, Inc. All Rights Reserved.
+#pragma once
+
+#include "Containers/ChunkedArray.h"
+#include "SDCEToken.h"
+#include <cctype>
+
+namespace UE::ShaderMinifier::SDCE
+{
+	/** We're making heavy assumptions on the ansi bit-width */
+	static_assert(std::is_same_v<FLookupCharType, char>, "SDCE-Tokenizer assumes ansi characters for lookups");
+
+	/** Array size needed for a bit-width */
+	template<typename T>
+	constexpr uint32 TArraySizeFor = 1u << (sizeof(T) * 8);
+
+	/** Mask size needed for a bit-width */
+	template<typename T>
+	constexpr uint32 TMask32SizeFor = TArraySizeFor<T> / 32;
+
+	/** Invalid operation head */
+	static constexpr int32_t TokenizerInvalidOpHead = INT32_MAX;
+
+	struct FTokenizerSegmentAnsiNode
+	{
+		/** Token at this ansi node */
+		ETokenType Type = ETokenType::Unknown;
+
+		/** Child nodes */
+		FTokenizerSegmentAnsiNode* Trees[TArraySizeFor<FLookupCharType>]{};
+	};
+	
+	struct FTokenizerStatic
+	{
+		FTokenizerStatic()
+		{
+			// Populate all segments
+#define X(N, S) AddSegment(ETokenType::N, S);
+			UE_MINIFIER_SDCE_KEYWORD_TABLE
+#undef X
+
+			// Numeric 
+			for (char C = '0'; C <= '9'; ++C)
+			{
+				AddCharBit(AlphaNumericMask, C);
+			}
+
+			// Lowercase alpha
+			for (char C = 'a'; C <= 'z'; ++C)
+			{
+				AddCharBit(AlphaNumericMask, C);
+				AddCharBit(AlphaMask, C);
+			}
+
+			// Uppercase alpha
+			for (char C = 'A'; C <= 'Z'; ++C)
+			{
+				AddCharBit(AlphaNumericMask, C);
+				AddCharBit(AlphaMask, C);
+			}
+
+			// All alpha's accept _
+			AddCharBit(AlphaNumericMask, '_');
+			AddCharBit(AlphaMask, '_');
+
+			// All white-space like characters
+			AddCharBit(SpaceMask, ' ');
+			AddCharBit(SpaceMask, '\f');
+			AddCharBit(SpaceMask, '\r');
+			AddCharBit(SpaceMask, '\n');
+			AddCharBit(SpaceMask, '\t');
+			AddCharBit(SpaceMask, '\v');
+		}
+
+		void AddSegment(ETokenType Type, const FLookupViewType& Value)
+		{
+			FTokenizerSegmentAnsiNode* Base = &SegmentRoot;
+
+			// Populate the segment node tree for the full length
+			for (int32 i = 0; i < Value.Len(); i++)
+			{
+				FLookupCharType Char = Value[i];
+
+				uint8 Index = static_cast<uint8>(Char);
+				if (!Base->Trees[Index])
+				{
+					Base->Trees[Index] = &AnsiNodeAllocator[AnsiNodeAllocator.Emplace()];
+				}
+
+				Base = Base->Trees[Index];
+			}
+
+			Base->Type = Type;
+		}
+
+		static void AddCharBit(uint32_t Set[], FLookupCharType Value)
+		{
+			uint8 Index = static_cast<uint8>(Value);
+			Set[Index / 32] |= 1u << (Index % 32);
+		}
+
+		static bool GetCharBit(uint32_t Set[], FLookupCharType Value)
+		{
+			uint8 Index = static_cast<uint8>(Value);
+			return Set[Index / 32] & (1u << (Index % 32));
+		}
+
+		static bool IsDigit(FLookupCharType Ch)
+		{
+			return Ch >= '0' && Ch <= '9';
+		}
+
+		bool IsAlphaNumeric(FLookupCharType Ch)
+		{
+			return GetCharBit(AlphaNumericMask, Ch);
+		}
+
+		bool IsAlpha(FLookupCharType Ch)
+		{
+			return GetCharBit(AlphaMask, Ch);
+		}
+
+		bool IsSpace(FLookupCharType Ch)
+		{
+			return GetCharBit(SpaceMask, Ch);
+		}
+
+		// Lookup masks
+		uint32_t AlphaMask       [TMask32SizeFor<FLookupCharType>]{};
+		uint32_t AlphaNumericMask[TMask32SizeFor<FLookupCharType>]{};
+		uint32_t SpaceMask       [TMask32SizeFor<FLookupCharType>]{};
+
+		// Base segment
+		FTokenizerSegmentAnsiNode SegmentRoot;
+
+	private:
+		// Chunked node allocator for coherence
+		TChunkedArray<FTokenizerSegmentAnsiNode, sizeof(FTokenizerSegmentAnsiNode) * 256> AnsiNodeAllocator;
+	} static TokenizerStatic;
+
+	struct FShallowTokenizer
+	{
+		FShallowTokenizer() = default;
+
+		/** Move constructor */
+		FShallowTokenizer(FShallowTokenizer&& rhs) : Token(rhs.Token), OpHead(rhs.OpHead), Begin(rhs.Begin), End(rhs.End)
+		{
+			// Invalidate RHS op-head
+			rhs.OpHead = TokenizerInvalidOpHead;
+		}
+
+		~FShallowTokenizer()
+		{
+			// If there's a valid op head, this is a branch that hasn't been accepted/rejected
+			check(OpHead == TokenizerInvalidOpHead);
+		}
+
+		/** Branch this tokenizer */
+		FShallowTokenizer Branch() const
+		{
+			FShallowTokenizer Branch;
+			Branch.OpHead = OpHead;
+			Branch.Token = Token;
+			Branch.Begin = Begin;
+			Branch.End = End;
+			return Branch;
+		}
+
+		/** Accept another tokenizer onto this */
+		void Accept(FShallowTokenizer& RHS)
+		{
+			// We keep our own op-head, but reset the RHS
+			RHS.OpHead = TokenizerInvalidOpHead;
+
+			// Inherit token cursor
+			Token = RHS.Token;
+			Begin = RHS.Begin;
+		}
+
+		/** Check if the current token type is <Type> */
+		bool Is(ETokenType Type)
+		{
+			return Token.Type == Type;
+		}
+
+		/** Check if the current token type is <Type> and consume it */
+		bool IsConsume(ETokenType Type)
+		{
+			if (Is(Type))
+			{
+				Next();
+				return true;
+			}
+
+			return false;
+		}
+
+		/** Get the current token and proceed */
+		FToken Next()
+		{
+			FToken CurrentToken = Token;
+			SkipWhitespace();
+			NextInternal();
+			return CurrentToken;
+		}
+
+		/** Actual tokenization internals */
+		void NextInternal()
+		{
+			FParseCharType* TokenStart = Begin;
+			if (Begin >= End)
+			{
+				return;
+			}
+
+			// Resolve static lifetime once
+			FTokenizerStatic& Static = TokenizerStatic;
+
+			// Numeric?
+			if (Static.IsDigit(*Begin) || (*Begin == '.' && Static.IsDigit(PeekInternal(Begin, 1))))
+			{
+				bool bHasPunctuation = *Begin == '.';
+				if (bHasPunctuation)
+				{
+					Begin++;
+				}
+
+				// Parse through the number
+				while (Begin < End && (Static.IsDigit(*Begin) || (!bHasPunctuation && *Begin == '.')))
+				{
+					Begin++;
+				}
+
+				// Handle specializations
+				if (Begin < End)
+				{
+					if (*Begin == 'x')
+					{
+						Begin++;
+						while (Begin < End && Static.IsAlphaNumeric(*Begin))
+						{
+							Begin++;
+						}
+					}
+					else if (*Begin == 'b')
+					{
+						Begin++;
+						while (Begin < End && Static.IsDigit(*Begin))
+						{
+							Begin++;
+						}
+					}
+					else if (*Begin == 'e')
+					{
+						Begin++;
+
+						if (*Begin == '+' || *Begin == '-')
+						{
+							Begin++;
+						}
+				
+						while (Begin < End && Static.IsDigit(*Begin))
+						{
+							Begin++;
+						}
+					}
+					else if (*Begin == '#')
+					{
+						Begin++;
+
+						while (Begin < End && Static.IsAlpha(*Begin))
+						{
+							Begin++;
+						}
+					}
+				}
+
+				// Handle typing
+				if (Begin < End)
+				{
+					switch (std::tolower(static_cast<uint8>(*Begin)))
+					{
+					default:
+						break;
+					case 'u':
+					case 'f':
+					case 'd':
+						Begin++;
+					}
+				}
+
+				Token.Type = ETokenType::Numeric;
+				Token.Begin = TokenStart;
+				Token.Length = Begin - TokenStart;
+				Token.XXHash32 = FSmallStringView::Hash(TokenStart, Token.Length);
+				return;
+			}
+
+			// Current segment tree node
+			FTokenizerSegmentAnsiNode* Tree = &Static.SegmentRoot;
+
+			// Is this a potential segment?
+			if (Tree->Trees[static_cast<uint8>(*Begin)])
+			{
+				FParseCharType* Shadow = Begin;
+
+				// Keep traversing until we're out of nodes
+				while (Shadow < End && Tree->Trees[static_cast<uint8>(*Shadow)])
+				{
+					Tree = Tree->Trees[static_cast<uint8>(*Shadow)];
+					Shadow++;
+				}
+
+				// If the last character was alnum, the next cannot be
+				if (!Static.IsAlphaNumeric(*(Shadow - 1)) || (!Static.IsAlphaNumeric(*Shadow)))
+				{
+					if (Tree->Type != ETokenType::Unknown)
+					{
+						Token = {};
+						Token.Begin = TokenStart;
+						Token.Type = Tree->Type;
+						Begin = Shadow;
+						return;
+					}
+				}
+			}
+
+			// Not a segment, just a regular identifier?
+			if (Static.IsAlpha(*Begin))
+			{
+				while (Begin < End && Static.IsAlphaNumeric(*Begin))
+				{
+					Begin++;
+				}
+
+				Token.Type = ETokenType::ID;
+				Token.Begin = TokenStart;
+				Token.Length = Begin - TokenStart;
+				Token.XXHash32 = FSmallStringView::Hash(TokenStart, Token.Length);
+				return;
+			}
+
+			// End of stream?
+			if (Begin >= End)
+			{
+				Token.Begin = End - 1;
+				Token.Type = ETokenType::Unknown;
+				Token.SetIDNoHash(FParseViewType(End - 1, 1));
+				return;
+			}
+
+			// Unknown token, report as single character
+			Token.Begin = TokenStart;
+			Token.Type = ETokenType::Unknown;
+			Token.SetIDNoHash(FParseViewType(TokenStart, ++Begin - TokenStart));
+		}
+
+		/** Peek a character */
+		FParseCharType PeekInternal(const FParseCharType* Cursor, uint32 Offset = 0)
+		{
+			if (Cursor + Offset >= End)
+			{
+				return 0;
+			}
+
+			return *(Cursor + Offset);
+		}
+
+		/** Skip until we meet a character */
+		bool SkipUntil(FParseCharType Ch)
+		{
+			while (Begin < End && *Begin != Ch)
+			{
+				Begin++;
+			}
+		
+			bool bIs = *Begin == Ch;
+			Next();
+			return bIs;
+		}
+
+		/** Skip until we meet a set of characters */
+		bool SkipUntil(FParseConstViewType ChSet)
+		{
+			while (Begin < End && !IsOneOf(*Begin, ChSet))
+			{
+				Begin++;
+			}
+
+			bool bIs = IsOneOf(*Begin, ChSet);
+			Next();
+			return bIs;
+		}
+
+		/** Skip all whitespaces */
+		void SkipWhitespace()
+		{
+			while (Begin < End)
+			{
+				if (TokenizerStatic.IsSpace(*Begin))
+				{
+					Begin++;
+				}
+
+				// We're cheating, also skip preprocessor directives here
+				else if (*Begin == '#')
+				{
+					while (Begin < End && *Begin != '\n')
+					{
+						Begin++;
+					}
+				}
+
+				// Not a whitespace
+				else
+				{
+					break;
+				}
+			}
+		}
+
+		/** Is this end of stream? */
+		bool IsEOS()
+		{
+			SkipWhitespace();
+			return Begin >= End;
+		}
+
+		/** Check if a character is in a set */
+		static bool IsOneOf(char Ch, FParseConstViewType ChSet)
+		{
+			for (char Set : ChSet)
+			{
+				if (Set == Ch)
+				{
+					return true;
+				}
+			}
+		
+			return false;
+		}
+
+		/** Current cursor offset */
+		uint32_t Offset() const
+		{
+			return End - Begin;
+		}
+
+		/** Current token */
+		FToken Token;
+
+		/** Operation buffer head, used for parser branching logic */
+		int32 OpHead = 0;
+
+		/** Cursors */
+		FParseCharType* Begin = nullptr;
+		FParseCharType* End   = nullptr;
+	};
+}
diff --git a/Engine/Source/Developer/ShaderCompilerCommon/Private/ShaderSDCE.cpp b/Engine/Source/Developer/ShaderCompilerCommon/Private/ShaderSDCE.cpp
new file mode 100644
--- /dev/null
+++ b/Engine/Source/Developer/ShaderCompilerCommon/Private/ShaderSDCE.cpp
@@ -0,0 +1,63 @@
+﻿// Copyright Epic Games, Inc. All Rights Reserved.
+
+#include "ShaderSDCE.h"
+#include "SDCE/SDCEContext.h"
+
+namespace UE::ShaderMinifier::SDCE
+{
+	void MinifyInPlace(const TConstArrayView<FStringView>& DCESymbols, FString& Code)
+	{
+		FMemMark Mark(FMemStack::Get());
+		FContext SDCEContext(DCESymbols);
+
+		// Parse and extract semantics
+		SDCEContext.Analyze(Code);
+
+		// Minify in-place to the existing chunks
+		SDCEContext.MinifyInPlace();
+	}
+	
+	bool MinifyInPlace(const FShaderCompilerInput& Input, const FShaderPreprocessOutput& PreprocessOutput, FString& Code)
+	{
+		TArray<FString, TInlineAllocator<32>> DCESymbols;
+
+		// Extract symbols from pragmas
+		PreprocessOutput.VisitDirectivesWithPrefix(ShaderMetadataPrefix,  [&DCESymbols](const FString* Value)
+		{
+			DCESymbols.Add(Value->RightChop(FCString::Strlen(ShaderMetadataPrefix)));
+		});
+
+		// Top environment
+		FString SDCESymbolsStr;
+		if (Input.Environment.GetCompileArgument(TEXT("UESHADERMETADATA_SDCE"), SDCESymbolsStr))
+		{
+			TArray<FString> EnvironmentSymbols;
+			SDCESymbolsStr.ParseIntoArray(EnvironmentSymbols, TEXT(";"));
+			DCESymbols.Append(EnvironmentSymbols);
+		}
+
+		// Shared environment
+		if (IsValidRef(Input.SharedEnvironment) && Input.SharedEnvironment->GetCompileArgument(TEXT("UESHADERMETADATA_SDCE"), SDCESymbolsStr))
+		{
+			TArray<FString> EnvironmentSymbols;
+			SDCESymbolsStr.ParseIntoArray(EnvironmentSymbols, TEXT(";"));
+			DCESymbols.Append(EnvironmentSymbols);
+		}
+
+		// Nothing to minify?
+		if (DCESymbols.IsEmpty())
+		{
+			return false;
+		}
+
+		// To views
+		TArray<FStringView, TInlineAllocator<32>> DCESymbolViews;
+		for (const FParseStringType& Symbol : DCESymbols)
+		{
+			DCESymbolViews.Add(Symbol);
+		}
+
+		MinifyInPlace(DCESymbolViews, Code);
+		return true;
+	}
+}
diff --git a/Engine/Source/Developer/ShaderCompilerCommon/Public/ShaderSDCE.h b/Engine/Source/Developer/ShaderCompilerCommon/Public/ShaderSDCE.h
new file mode 100644
--- /dev/null
+++ b/Engine/Source/Developer/ShaderCompilerCommon/Public/ShaderSDCE.h
@@ -0,0 +1,17 @@
+// Copyright Epic Games, Inc. All Rights Reserved.
+#pragma once
+
+#include "ShaderPreprocessTypes.h"
+
+namespace UE::ShaderMinifier::SDCE
+{
+	/**
+	 * Minify shader sources
+	 */
+	SHADERCOMPILERCOMMON_API void MinifyInPlace(const TConstArrayView<FStringView>& DCESymbols, FString& Code);
+	
+	/**
+	 * Minify shader sources from preprocessor output
+	 */
+	SHADERCOMPILERCOMMON_API bool MinifyInPlace(const FShaderCompilerInput& Input, const FShaderPreprocessOutput& PreprocessOutput, FString& Code);
+}
diff --git a/Engine/Source/Developer/Windows/ShaderFormatD3D/Private/D3DShaderCompiler.cpp b/Engine/Source/Developer/Windows/ShaderFormatD3D/Private/D3DShaderCompiler.cpp
--- a/Engine/Source/Developer/Windows/ShaderFormatD3D/Private/D3DShaderCompiler.cpp
+++ b/Engine/Source/Developer/Windows/ShaderFormatD3D/Private/D3DShaderCompiler.cpp
@@ -16,6 +16,7 @@
 #include "SpirvCommon.h"
 #include "Algo/Transform.h"
 #include "HlslParserInternal.h"
+#include "ShaderSDCE.h"
 
 #define DEBUG_SHADERS 0
 
@@ -1466,7 +1467,10 @@ void CompileD3DShader(const FShaderCompilerInput& Input, const FShaderPreprocess
 		return;
 	}
 
-	if (ShaderParameterParser.DidModifyShader())
+	// Perform symbolic elimination
+	bool bSDCEMinified = UE::ShaderMinifier::SDCE::MinifyInPlace(Input, InPreprocessOutput, PreprocessedSource);
+
+	if (ShaderParameterParser.DidModifyShader() || bSDCEMinified)
 	{
 		Output.ModifiedShaderSource = PreprocessedSource;
 	}
diff --git a/Engine/Source/Runtime/Engine/Private/Materials/HLSLMaterialTranslator.cpp b/Engine/Source/Runtime/Engine/Private/Materials/HLSLMaterialTranslator.cpp
--- a/Engine/Source/Runtime/Engine/Private/Materials/HLSLMaterialTranslator.cpp
+++ b/Engine/Source/Runtime/Engine/Private/Materials/HLSLMaterialTranslator.cpp
@@ -3092,6 +3092,27 @@ void FHLSLMaterialTranslator::GetMaterialEnvironment(EShaderPlatform InPlatform,
 		OutEnvironment.SetDefine(TEXT("SUBSTRATE_FASTPATH"), EnvironmentDefines->bSubstrateFastPath);
 		OutEnvironment.SetDefine(TEXT("SUBSTRATE_CLAMPED_CLOSURE_COUNT"), EnvironmentDefines->SubstrateClampedClosureCount);
 		OutEnvironment.SetDefine(TEXT("SUBSTRATE_COMPLEXSPECIALPATH"), EnvironmentDefines->bSubstrateComplexSpecialPath);
+
+		if (!Substrate::IsSubstrateBlendableGBufferEnabled(InPlatform))
+		{
+			static const TCHAR* SDCESymbols[] = {
+				TEXT("FMaterialPixelParameters"),
+				TEXT("FMaterialLWCData"),
+				TEXT("FPrimitiveSceneData"),
+				TEXT("FPixelMaterialInputs"),
+				TEXT("FCluster"),
+				TEXT("FNaniteView"),
+				TEXT("FInstanceDynamicData"),
+				TEXT("FNanitePostDeformVertex"),
+				TEXT("FMaterialVertexParameters"),
+				TEXT("FNanitePixelAttributes"),
+				TEXT("ViewState"),
+				TEXT("FSceneData"),
+				TEXT("FPrimitiveSceneData")
+			};
+		
+			OutEnvironment.SetCompileArgument(TEXT("UESHADERMETADATA_SDCE"), FString::Join(SDCESymbols, TEXT(";")));
+		}
 	}
 
 	for (int i = 0; i < EnvironmentDefines->SubstrateDefines.Num(); ++i)
