Shader SDCE, moved substrate symbols to compiler argument
- Only applies to new translated materials, does not invalidate existing

diff --git a/Engine/Source/Developer/ShaderCompilerCommon/Private/ShaderSDCE.cpp b/Engine/Source/Developer/ShaderCompilerCommon/Private/ShaderSDCE.cpp
--- a/Engine/Source/Developer/ShaderCompilerCommon/Private/ShaderSDCE.cpp
+++ b/Engine/Source/Developer/ShaderCompilerCommon/Private/ShaderSDCE.cpp
@@ -18,15 +18,33 @@ namespace UE::ShaderMinifier::SDCE
 		SDCEContext.MinifyInPlace();
 	}
 	
-	bool MinifyInPlace(const FShaderPreprocessOutput& PreprocessOutput, FString& Code)
+	bool MinifyInPlace(const FShaderCompilerInput& Input, const FShaderPreprocessOutput& PreprocessOutput, FString& Code)
 	{
-		// Get symbols to DCE
 		TArray<FString, TInlineAllocator<32>> DCESymbols;
+
+		// Extract symbols from pragmas
 		PreprocessOutput.VisitDirectivesWithPrefix(ShaderMetadataPrefix,  [&DCESymbols](const FString* Value)
 		{
 			DCESymbols.Add(Value->RightChop(FCString::Strlen(ShaderMetadataPrefix)));
 		});
 
+		// Top environment
+		FString SDCESymbolsStr;
+		if (Input.Environment.GetCompileArgument(TEXT("UESHADERMETADATA_SDCE"), SDCESymbolsStr))
+		{
+			TArray<FString> EnvironmentSymbols;
+			SDCESymbolsStr.ParseIntoArray(EnvironmentSymbols, TEXT(";"));
+			DCESymbols.Append(EnvironmentSymbols);
+		}
+
+		// Shared environment
+		if (IsValidRef(Input.SharedEnvironment) && Input.SharedEnvironment->GetCompileArgument(TEXT("UESHADERMETADATA_SDCE"), SDCESymbolsStr))
+		{
+			TArray<FString> EnvironmentSymbols;
+			SDCESymbolsStr.ParseIntoArray(EnvironmentSymbols, TEXT(";"));
+			DCESymbols.Append(EnvironmentSymbols);
+		}
+
 		// Nothing to minify?
 		if (DCESymbols.IsEmpty())
 		{
diff --git a/Engine/Source/Developer/ShaderCompilerCommon/Public/ShaderSDCE.h b/Engine/Source/Developer/ShaderCompilerCommon/Public/ShaderSDCE.h
--- a/Engine/Source/Developer/ShaderCompilerCommon/Public/ShaderSDCE.h
+++ b/Engine/Source/Developer/ShaderCompilerCommon/Public/ShaderSDCE.h
@@ -13,5 +13,5 @@ namespace UE::ShaderMinifier::SDCE
 	/**
 	 * Minify shader sources from preprocessor output
 	 */
-	SHADERCOMPILERCOMMON_API bool MinifyInPlace(const FShaderPreprocessOutput& PreprocessOutput, FString& Code);
+	SHADERCOMPILERCOMMON_API bool MinifyInPlace(const FShaderCompilerInput& Input, const FShaderPreprocessOutput& PreprocessOutput, FString& Code);
 }
diff --git a/Engine/Source/Developer/Windows/ShaderFormatD3D/Private/D3DShaderCompiler.cpp b/Engine/Source/Developer/Windows/ShaderFormatD3D/Private/D3DShaderCompiler.cpp
--- a/Engine/Source/Developer/Windows/ShaderFormatD3D/Private/D3DShaderCompiler.cpp
+++ b/Engine/Source/Developer/Windows/ShaderFormatD3D/Private/D3DShaderCompiler.cpp
@@ -1468,7 +1468,7 @@ void CompileD3DShader(const FShaderCompilerInput& Input, const FShaderPreprocess
 	}
 
 	// Perform symbolic elimination
-	bool bSDCEMinified = UE::ShaderMinifier::SDCE::MinifyInPlace(InPreprocessOutput, PreprocessedSource);
+	bool bSDCEMinified = UE::ShaderMinifier::SDCE::MinifyInPlace(Input, InPreprocessOutput, PreprocessedSource);
 
 	if (ShaderParameterParser.DidModifyShader() || bSDCEMinified)
 	{
diff --git a/Engine/Source/Runtime/Engine/Private/Materials/HLSLMaterialTranslator.cpp b/Engine/Source/Runtime/Engine/Private/Materials/HLSLMaterialTranslator.cpp
--- a/Engine/Source/Runtime/Engine/Private/Materials/HLSLMaterialTranslator.cpp
+++ b/Engine/Source/Runtime/Engine/Private/Materials/HLSLMaterialTranslator.cpp
@@ -3086,6 +3086,24 @@ void FHLSLMaterialTranslator::GetMaterialEnvironment(EShaderPlatform InPlatform,
 		OutEnvironment.SetDefine(TEXT("SUBSTRATE_FASTPATH"), EnvironmentDefines->bSubstrateFastPath);
 		OutEnvironment.SetDefine(TEXT("SUBSTRATE_CLAMPED_CLOSURE_COUNT"), EnvironmentDefines->SubstrateClampedClosureCount);
 		OutEnvironment.SetDefine(TEXT("SUBSTRATE_COMPLEXSPECIALPATH"), EnvironmentDefines->bSubstrateComplexSpecialPath);
+
+		static const TCHAR* SDCESymbols[] = {
+			TEXT("FMaterialPixelParameters"),
+			TEXT("FMaterialLWCData"),
+			TEXT("FPrimitiveSceneData"),
+			TEXT("FPixelMaterialInputs"),
+			TEXT("FCluster"),
+			TEXT("FNaniteView"),
+			TEXT("FInstanceDynamicData"),
+			TEXT("FNanitePostDeformVertex"),
+			TEXT("FMaterialVertexParameters"),
+			TEXT("FNanitePixelAttributes"),
+			TEXT("ViewState"),
+			TEXT("FSceneData"),
+			TEXT("FPrimitiveSceneData")
+		};
+		
+		OutEnvironment.SetCompileArgument(TEXT("UESHADERMETADATA_SDCE"), FString::Join(SDCESymbols, TEXT(";")));
 	}
 
 	for (int i = 0; i < EnvironmentDefines->SubstrateDefines.Num(); ++i)
