Shader SDCE, removed invalid pragma once

diff --git a/Engine/Source/Developer/ShaderCompilerCommon/Private/ShaderSDCE.cpp b/Engine/Source/Developer/ShaderCompilerCommon/Private/ShaderSDCE.cpp
--- a/Engine/Source/Developer/ShaderCompilerCommon/Private/ShaderSDCE.cpp
+++ b/Engine/Source/Developer/ShaderCompilerCommon/Private/ShaderSDCE.cpp
@@ -1,5 +1,4 @@
 ﻿// Copyright Epic Games, Inc. All Rights Reserved.
-#pragma once
 
 #include "ShaderSDCE.h"
 #include "SDCE/SDCEContext.h"
