ShaderCompilerCommon changes
- Passing UniformBufferName around as a string view to reduce copies.
- UniformBufferName needs to be adjusted before checking if it's in the UniformBufferMap.

diff --git a/Engine/Source/Developer/ShaderCompilerCommon/Private/ShaderCompilerCommon.cpp b/Engine/Source/Developer/ShaderCompilerCommon/Private/ShaderCompilerCommon.cpp
--- a/Engine/Source/Developer/ShaderCompilerCommon/Private/ShaderCompilerCommon.cpp
+++ b/Engine/Source/Developer/ShaderCompilerCommon/Private/ShaderCompilerCommon.cpp
@@ -102,7 +102,9 @@ void UE::ShaderCompilerCommon::BuildShaderResourceTable(const FShaderCompilerRes
 
 static bool DoesUniformBufferNeedReflectedMembers(const TMap<FString, FUniformBufferEntry>& UniformBufferMap, FStringView UniformBufferName)
 {
-	if (const FUniformBufferEntry* Entry = UniformBufferMap.FindByHash(GetTypeHash(UniformBufferName), UniformBufferName))
+	FStringView AdjustedUniformBufferName(UE::ShaderCompilerCommon::RemoveConstantBufferPrefix(UniformBufferName));
+
+	if (const FUniformBufferEntry* Entry = UniformBufferMap.FindByHash(GetTypeHash(AdjustedUniformBufferName), AdjustedUniformBufferName))
 	{
 		return EnumHasAnyFlags(Entry->Flags, ERHIUniformBufferFlags::NeedsReflectedMembers);
 	}
@@ -667,14 +669,14 @@ void HandleReflectedRootConstantBuffer(
 }
 
 void HandleReflectedUniformBuffer(
-	const FString& UniformBufferName,
+	FStringView UniformBufferName,
 	int32 ReflectionSlot,
 	int32 BaseIndex,
 	int32 BufferSize,
 	FShaderCompilerOutput& CompilerOutput
 )
 {
-	FString AdjustedUniformBufferName(UE::ShaderCompilerCommon::RemoveConstantBufferPrefix(UniformBufferName));
+	FStringView AdjustedUniformBufferName(UE::ShaderCompilerCommon::RemoveConstantBufferPrefix(UniformBufferName));
 
 	CompilerOutput.ParameterMap.AddParameterAllocation(
 		AdjustedUniformBufferName,
diff --git a/Engine/Source/Developer/ShaderCompilerCommon/Public/ShaderCompilerCommon.h b/Engine/Source/Developer/ShaderCompilerCommon/Public/ShaderCompilerCommon.h
--- a/Engine/Source/Developer/ShaderCompilerCommon/Public/ShaderCompilerCommon.h
+++ b/Engine/Source/Developer/ShaderCompilerCommon/Public/ShaderCompilerCommon.h
@@ -240,19 +240,19 @@ extern SHADERCOMPILERCOMMON_API void HandleReflectedRootConstantBuffer(
 );
 
 extern SHADERCOMPILERCOMMON_API void HandleReflectedUniformBuffer(
-	const FString& UniformBufferName,
+	FStringView UniformBufferName,
 	int32 ReflectionSlot,
 	int32 BaseIndex,
 	int32 BufferSize,
 	FShaderCompilerOutput& CompilerOutput
 );
 
-inline void HandleReflectedUniformBuffer(const FString& UniformBufferName, int32 ReflectionSlot, int32 BufferSize, FShaderCompilerOutput& CompilerOutput)
+inline void HandleReflectedUniformBuffer(FStringView UniformBufferName, int32 ReflectionSlot, int32 BufferSize, FShaderCompilerOutput& CompilerOutput)
 {
 	HandleReflectedUniformBuffer(UniformBufferName, ReflectionSlot, 0, BufferSize, CompilerOutput);
 }
 
-inline void HandleReflectedUniformBuffer(const FString& UniformBufferName, int32 ReflectionSlot, FShaderCompilerOutput& CompilerOutput)
+inline void HandleReflectedUniformBuffer(FStringView UniformBufferName, int32 ReflectionSlot, FShaderCompilerOutput& CompilerOutput)
 {
 	HandleReflectedUniformBuffer(UniformBufferName, ReflectionSlot, 0, CompilerOutput);
 }
