Fixing spelling in comment.

diff --git a/Engine/Source/Developer/ShaderCompilerCommon/Private/ShaderCompilerCommon.cpp b/Engine/Source/Developer/ShaderCompilerCommon/Private/ShaderCompilerCommon.cpp
--- a/Engine/Source/Developer/ShaderCompilerCommon/Private/ShaderCompilerCommon.cpp
+++ b/Engine/Source/Developer/ShaderCompilerCommon/Private/ShaderCompilerCommon.cpp
@@ -2288,7 +2288,7 @@ static void DumpDebugShaderDisassembled(const FShaderCompilerInput& Input, Cross
 		TArray<ANSICHAR> AssemblyText;
 		if (CrossCompiler::FShaderConductorContext::Disassemble(Language, InData, InDataByteSize, AssemblyText))
 		{
-			// Assembly text contains NUL terminator, so text lenght is |array|-1
+			// Assembly text contains NUL terminator, so text length is |array|-1
 			DumpDebugShaderText(Input, AssemblyText.GetData(), AssemblyText.Num() - 1, FileExtension);
 		}
 	}
