Fixing issues with UniformBuffer definition parsing
- Don't move back one character when using next separator/semicolon when calculating the "end" of the string since that should always be one character past the last character of the string.
- When checking for 1:1 name changes, don't skip changes that result in smaller strings.
- Adding code to make it easier to debug the definitions.

#Synced-CL 48096956

diff --git a/Engine/Source/Developer/ShaderCompilerCommon/Private/ShaderCompilerCommon.cpp b/Engine/Source/Developer/ShaderCompilerCommon/Private/ShaderCompilerCommon.cpp
--- a/Engine/Source/Developer/ShaderCompilerCommon/Private/ShaderCompilerCommon.cpp
+++ b/Engine/Source/Developer/ShaderCompilerCommon/Private/ShaderCompilerCommon.cpp
@@ -954,10 +954,10 @@ const FShaderSource::CharType* ParseUniformBufferDefinition(const FShaderSource:
 		if (NextSeparator < NextSemicolon)
 		{
 			const FShaderSource::CharType* StructStart = CurrentParseStart;
-			const FShaderSource::CharType* StructEnd = NextSeparator - 1;
+			const FShaderSource::CharType* StructEnd = NextSeparator;
 
 			const FShaderSource::CharType* GlobalStart = NextSeparator + 1;
-			const FShaderSource::CharType* GlobalEnd = NextSemicolon - 1;
+			const FShaderSource::CharType* GlobalEnd = NextSemicolon;
 
 			while (IsSpaceOrTabOrEOL(*StructStart))
 			{
@@ -975,7 +975,7 @@ const FShaderSource::CharType* ParseUniformBufferDefinition(const FShaderSource:
 			FShaderSource::FViewType GlobalName(GlobalStart, GlobalEnd - GlobalStart);
 
 			// Avoid unnecessary conversions
-			if (StructName.Len() == GlobalName.Len() && FShaderSource::FCStringType::Strncmp(StructName.GetData(), GlobalName.GetData(), StructName.Len()) != 0)
+			if (StructName.Len() != GlobalName.Len() || FShaderSource::FCStringType::Strncmp(StructName.GetData(), GlobalName.GetData(), StructName.Len()) != 0)
 			{
 				FUniformBufferMemberInfoNew NewMemberInfo;
 				NewMemberInfo.NameAsStructMember = StructName;
@@ -1443,25 +1443,29 @@ void CleanupUniformBufferCode(const FShaderCompilerEnvironment& Environment, FSh
 
 	} while (bUniformBufferFound);
 
-	// Compact commented out uniform buffers out of the output source.  This costs around 10x less to do here than later in the minifier.  Note that
-	// it's not necessary to add a line directive to fix up line numbers because uniform buffer declarations are always in generated files, and there
-	// will be a line directive already there for the transition from the generated file back to whatever file included it.  The destination offset
-	// for the first move is the start of the first uniform buffer declaration we are overwriting, then advances as characters are copied.
-	int32 DestOffset = UniformBufferSpans.Num() ? UniformBufferSpans[0].Offset : PreprocessedShaderSource.Len();
-
-	for (int32 SpanIndex = 0; SpanIndex < UniformBufferSpans.Num(); SpanIndex++)
+	// Uncomment this if to keep the UniformBuffer{} definitions in the preprocessed code.
+	//if (false)
 	{
-		// The source code we are compacting down is from the end of one span to the start of the next span, or end of the string.
-		// We do not need to account for null terminator as the ShrinkToLen call below will null terminate for us.
-		int32 SourceOffset = UniformBufferSpans[SpanIndex].Offset + UniformBufferSpans[SpanIndex].Length;
-		int32 MoveCount = (SpanIndex < UniformBufferSpans.Num() - 1 ? UniformBufferSpans[SpanIndex + 1].Offset : PreprocessedShaderSource.Len()) - SourceOffset;
+		// Compact commented out uniform buffers out of the output source.  This costs around 10x less to do here than later in the minifier.  Note that
+		// it's not necessary to add a line directive to fix up line numbers because uniform buffer declarations are always in generated files, and there
+		// will be a line directive already there for the transition from the generated file back to whatever file included it.  The destination offset
+		// for the first move is the start of the first uniform buffer declaration we are overwriting, then advances as characters are copied.
+		int32 DestOffset = UniformBufferSpans.Num() ? UniformBufferSpans[0].Offset : PreprocessedShaderSource.Len();
+
+		for (int32 SpanIndex = 0; SpanIndex < UniformBufferSpans.Num(); SpanIndex++)
+		{
+			// The source code we are compacting down is from the end of one span to the start of the next span, or end of the string.
+			// We do not need to account for null terminator as the ShrinkToLen call below will null terminate for us.
+			int32 SourceOffset = UniformBufferSpans[SpanIndex].Offset + UniformBufferSpans[SpanIndex].Length;
+			int32 MoveCount = (SpanIndex < UniformBufferSpans.Num() - 1 ? UniformBufferSpans[SpanIndex + 1].Offset : PreprocessedShaderSource.Len()) - SourceOffset;
 
-		check(DestOffset >= 0 && DestOffset < SourceOffset && SourceOffset + MoveCount <= PreprocessedShaderSource.Len());
+			check(DestOffset >= 0 && DestOffset < SourceOffset && SourceOffset + MoveCount <= PreprocessedShaderSource.Len());
 
-		memmove(SourceStart + DestOffset, SourceStart + SourceOffset, MoveCount * sizeof(FShaderSource::CharType));
-		DestOffset += MoveCount;
+			memmove(SourceStart + DestOffset, SourceStart + SourceOffset, MoveCount * sizeof(FShaderSource::CharType));
+			DestOffset += MoveCount;
+		}
+		PreprocessedShaderSource.ShrinkToLen(DestOffset, EAllowShrinking::No);
 	}
-	PreprocessedShaderSource.ShrinkToLen(DestOffset, EAllowShrinking::No);
 }
 
 ///////////////////////////////////////////////////////////////////////////////////////////////////////////////////////////////////////
