Moving new FShaderCompilerInput argument to be the first argument to match the rest of the ShaderCompilerCommon functions.

#Synced-CL 48800164

diff --git a/Engine/Source/Developer/Apple/MetalShaderFormat/Private/MetalCompileShaderMSC.cpp b/Engine/Source/Developer/Apple/MetalShaderFormat/Private/MetalCompileShaderMSC.cpp
--- a/Engine/Source/Developer/Apple/MetalShaderFormat/Private/MetalCompileShaderMSC.cpp
+++ b/Engine/Source/Developer/Apple/MetalShaderFormat/Private/MetalCompileShaderMSC.cpp
@@ -293,14 +293,14 @@ static void ProcessReflection(ID3D12ShaderReflection* ShaderReflection, const ui
 						if (VariableDesc.uFlags & D3D_SVF_USED)
 						{
 							HandleReflectedUniformBufferConstantBufferMember(
+								Input,
 								Reason,
 								UniformBufferName,
 								BindIndex,
 								FString(VariableDesc.Name),
 								VariableDesc.StartOffset,
 								VariableDesc.Size,
-								Output,
-								Input
+								Output
 							);
 						}
 					}
@@ -308,7 +308,7 @@ static void ProcessReflection(ID3D12ShaderReflection* ShaderReflection, const ui
 				
 				// Regular uniform buffer - we only care about the binding index
 				CCHeaderWriter.WriteUniformBlock(*UniformBufferName, BindIndex);
-				HandleReflectedUniformBuffer(UniformBufferName, BindIndex, Output, Input);
+				HandleReflectedUniformBuffer(Input, UniformBufferName, BindIndex, Output);
 			}
 			NumCBVs = FMath::Max(NumCBVs, BindIndex + BindDesc.BindCount);
 		}
diff --git a/Engine/Source/Developer/Apple/MetalShaderFormat/Private/MetalShaderCompiler.cpp b/Engine/Source/Developer/Apple/MetalShaderFormat/Private/MetalShaderCompiler.cpp
--- a/Engine/Source/Developer/Apple/MetalShaderFormat/Private/MetalShaderCompiler.cpp
+++ b/Engine/Source/Developer/Apple/MetalShaderFormat/Private/MetalShaderCompiler.cpp
@@ -390,7 +390,7 @@ void BuildMetalShaderOutput(
 		}
 		UsedUniformBufferSlots[UBIndex] = true;
 
-		HandleReflectedUniformBuffer(UniformBlock.Name, UBIndex, ShaderOutput, ShaderInput);
+		HandleReflectedUniformBuffer(ShaderInput, UniformBlock.Name, UBIndex, ShaderOutput);
 	}
 	
 	if (bOutOfBounds)
diff --git a/Engine/Source/Developer/ShaderCompilerCommon/Private/ShaderCompilerCommon.cpp b/Engine/Source/Developer/ShaderCompilerCommon/Private/ShaderCompilerCommon.cpp
--- a/Engine/Source/Developer/ShaderCompilerCommon/Private/ShaderCompilerCommon.cpp
+++ b/Engine/Source/Developer/ShaderCompilerCommon/Private/ShaderCompilerCommon.cpp
@@ -599,14 +599,14 @@ void HandleReflectedGlobalConstantBufferMember(
 }
 
 void HandleReflectedUniformBufferConstantBufferMember(
+	const FShaderCompilerInput& Input,
 	EUniformBufferMemberReflectionReason Reason,
 	FStringView UniformBufferName,
 	int32 UniformBufferSlot,
 	FStringView InMemberName,
 	int32 ReflectionOffset,
 	int32 ReflectionSize,
-	FShaderCompilerOutput& Output,
-	const FShaderCompilerInput& Input
+	FShaderCompilerOutput& Output
 )
 {
 	FStringView MemberName = InMemberName;
@@ -697,12 +697,12 @@ bool VerifyBoundUniform(const FShaderCompilerInput& Input, FStringView UniformBu
 }
 
 void HandleReflectedUniformBuffer(
+	const FShaderCompilerInput& Input,
 	FStringView UniformBufferName,
 	int32 ReflectionSlot,
 	int32 BaseIndex,
 	int32 BufferSize,
-	FShaderCompilerOutput& CompilerOutput,
-	const FShaderCompilerInput& Input
+	FShaderCompilerOutput& CompilerOutput
 )
 {
 	FStringView AdjustedUniformBufferName(UE::ShaderCompilerCommon::RemoveConstantBufferPrefix(UniformBufferName));
diff --git a/Engine/Source/Developer/ShaderCompilerCommon/Public/ShaderCompilerCommon.h b/Engine/Source/Developer/ShaderCompilerCommon/Public/ShaderCompilerCommon.h
--- a/Engine/Source/Developer/ShaderCompilerCommon/Public/ShaderCompilerCommon.h
+++ b/Engine/Source/Developer/ShaderCompilerCommon/Public/ShaderCompilerCommon.h
@@ -217,14 +217,14 @@ extern SHADERCOMPILERCOMMON_API void HandleReflectedGlobalConstantBufferMember(
 );
 
 extern SHADERCOMPILERCOMMON_API void HandleReflectedUniformBufferConstantBufferMember(
+	const FShaderCompilerInput& Input,
 	EUniformBufferMemberReflectionReason Reason,
 	FStringView UniformBufferName,
 	int32 UniformBufferSlot,
 	FStringView MemberName,
 	int32 ReflectionOffset,
 	int32 ReflectionSize,
-	FShaderCompilerOutput& CompilerOutput,
-	const FShaderCompilerInput& Input
+	FShaderCompilerOutput& CompilerOutput
 );
 
 extern SHADERCOMPILERCOMMON_API void HandleReflectedRootConstantBufferMember(
@@ -242,22 +242,22 @@ extern SHADERCOMPILERCOMMON_API void HandleReflectedRootConstantBuffer(
 );
 
 extern SHADERCOMPILERCOMMON_API void HandleReflectedUniformBuffer(
+	const FShaderCompilerInput& Input,
 	FStringView UniformBufferName,
 	int32 ReflectionSlot,
 	int32 BaseIndex,
 	int32 BufferSize,
-	FShaderCompilerOutput& CompilerOutput,
-	const FShaderCompilerInput& Input
+	FShaderCompilerOutput& CompilerOutput
 );
 
-inline void HandleReflectedUniformBuffer(FStringView UniformBufferName, int32 ReflectionSlot, int32 BufferSize, FShaderCompilerOutput& CompilerOutput, const FShaderCompilerInput& Input)
+inline void HandleReflectedUniformBuffer(const FShaderCompilerInput& Input, FStringView UniformBufferName, int32 ReflectionSlot, int32 BufferSize, FShaderCompilerOutput& CompilerOutput)
 {
-	HandleReflectedUniformBuffer(UniformBufferName, ReflectionSlot, 0, BufferSize, CompilerOutput, Input);
+	HandleReflectedUniformBuffer(Input, UniformBufferName, ReflectionSlot, 0, BufferSize, CompilerOutput);
 }
 
-inline void HandleReflectedUniformBuffer(FStringView UniformBufferName, int32 ReflectionSlot, FShaderCompilerOutput& CompilerOutput, const FShaderCompilerInput& Input)
+inline void HandleReflectedUniformBuffer(const FShaderCompilerInput& Input, FStringView UniformBufferName, int32 ReflectionSlot, FShaderCompilerOutput& CompilerOutput)
 {
-	HandleReflectedUniformBuffer(UniformBufferName, ReflectionSlot, 0, CompilerOutput, Input);
+	HandleReflectedUniformBuffer(Input, UniformBufferName, ReflectionSlot, 0, CompilerOutput);
 }
 
 extern SHADERCOMPILERCOMMON_API void HandleReflectedShaderResource(
diff --git a/Engine/Source/Developer/ShaderFormatOpenGL/Private/OpenGLShaderCompiler.cpp b/Engine/Source/Developer/ShaderFormatOpenGL/Private/OpenGLShaderCompiler.cpp
--- a/Engine/Source/Developer/ShaderFormatOpenGL/Private/OpenGLShaderCompiler.cpp
+++ b/Engine/Source/Developer/ShaderFormatOpenGL/Private/OpenGLShaderCompiler.cpp
@@ -673,7 +673,7 @@ void BuildShaderOutputInternal(
 		uint16 UBIndex = UniformBlock.Index;
 
 		UsedUniformBufferSlots[UBIndex] = true;
-		HandleReflectedUniformBuffer(UniformBlock.Name, UBIndex, ShaderOutput, ShaderInput);
+		HandleReflectedUniformBuffer(ShaderInput, UniformBlock.Name, UBIndex, ShaderOutput);
 		Header.Bindings.NumUniformBuffers++;
 	}
 
@@ -720,7 +720,7 @@ void BuildShaderOutputInternal(
 	for (auto& PackedUB : CCHeader.PackedUBs)
 	{
 		UsedUniformBufferSlots[PackedUB.Attribute.Index] = true;
-		HandleReflectedUniformBuffer(PackedUB.Attribute.Name, PackedUB.Attribute.Index, ShaderOutput, ShaderInput);
+		HandleReflectedUniformBuffer(ShaderInput, PackedUB.Attribute.Name, PackedUB.Attribute.Index, ShaderOutput);
 		Header.Bindings.NumUniformBuffers++;
 
 		// Nothing else...
diff --git a/Engine/Source/Developer/VulkanShaderFormat/Private/SpirVShaderCompiler.inl b/Engine/Source/Developer/VulkanShaderFormat/Private/SpirVShaderCompiler.inl
--- a/Engine/Source/Developer/VulkanShaderFormat/Private/SpirVShaderCompiler.inl
+++ b/Engine/Source/Developer/VulkanShaderFormat/Private/SpirVShaderCompiler.inl
@@ -806,7 +806,7 @@ static bool BuildShaderOutputFromSpirv(
 							if (InternalState.UseRootParametersStructure())
 							{
 								check(ReflectionSlot == FShaderParametersMetadata::kRootCBufferBindingIndex);
-								HandleReflectedUniformBuffer(ResourceName, ReflectionSlot, Output, InternalState.Input);
+								HandleReflectedUniformBuffer(InternalState.Input, ResourceName, ReflectionSlot, Output);
 							}
 
 							// Register all uniform buffer members of Globals as loose data
@@ -840,7 +840,7 @@ static bool BuildShaderOutputFromSpirv(
 						{
 							check(BindingTypeCount == ReflectionSlot);
 							check(!UsedUniformBufferSlots[ReflectionSlot]);
-							HandleReflectedUniformBuffer(ResourceName, ReflectionSlot, Output, InternalState.Input);
+							HandleReflectedUniformBuffer(InternalState.Input, ResourceName, ReflectionSlot, Output);
 							AddShaderValidationUBSize(BindingTypeCount, Binding->block.padded_size, Output);
 
 							const EUniformBufferMemberReflectionReason Reason = ShouldReflectUniformBufferMembers(InternalState.Input, ResourceName);
@@ -858,14 +858,14 @@ static bool BuildShaderOutputFromSpirv(
 
 									const FString MemberName(Member.name);
 									HandleReflectedUniformBufferConstantBufferMember(
+										InternalState.Input,
 										Reason,
 										ResourceName,
 										ReflectionSlot,
 										MemberName,
 										Member.absolute_offset,
 										Member.size,
-										Output,
-										InternalState.Input
+										Output
 									);
 								}
 							}
diff --git a/Engine/Source/Developer/Windows/ShaderFormatD3D/Private/D3DShaderCompiler.inl b/Engine/Source/Developer/Windows/ShaderFormatD3D/Private/D3DShaderCompiler.inl
--- a/Engine/Source/Developer/Windows/ShaderFormatD3D/Private/D3DShaderCompiler.inl
+++ b/Engine/Source/Developer/Windows/ShaderFormatD3D/Private/D3DShaderCompiler.inl
@@ -213,7 +213,7 @@ template <typename ID3D1xShaderReflection, typename D3D1x_SHADER_DESC, typename
 			{
 				// Track just the constant buffer itself.
 				AddShaderValidationUBSize(CBIndex, CBDesc.Size, Output);
-				HandleReflectedUniformBuffer(ConstantBufferName, CBIndex, Output, Input);
+				HandleReflectedUniformBuffer(Input, ConstantBufferName, CBIndex, Output);
 				
 				CompileData.UsedUniformBufferSlots[CBIndex] = true;
 
@@ -232,14 +232,14 @@ template <typename ID3D1xShaderReflection, typename D3D1x_SHADER_DESC, typename
 							const FString MemberName(VariableDesc.Name);
 
 							HandleReflectedUniformBufferConstantBufferMember(
+								Input,
 								Reason,
 								ConstantBufferName,
 								CBIndex,
 								MemberName,
 								VariableDesc.StartOffset,
 								VariableDesc.Size,
-								Output,
-								Input
+								Output
 							);
 						}
 					}
