VulkanRHI: Fix shader record reflection happening too early and picking up unused variables.  Build reflection from the final SPIRV where only the members being used of the VulkanShaderRecordLooseParameterData are detailed.  Re-enable PLATFORM_SUPPORTS_CONSTANTBUFFER_OBJECT.

#rnx

#Synced-CL 49995137

diff --git a/Engine/Shaders/Public/Platform/Vulkan/VulkanCommon.ush b/Engine/Shaders/Public/Platform/Vulkan/VulkanCommon.ush
--- a/Engine/Shaders/Public/Platform/Vulkan/VulkanCommon.ush
+++ b/Engine/Shaders/Public/Platform/Vulkan/VulkanCommon.ush
@@ -7,7 +7,7 @@
 #pragma once
 
 // Update this GUID to invalidate shader recompilation for only Vulkan shaders
-#pragma message("UESHADERMETADATA_VERSION B631D136-745C-4AD2-918E-3EBE61B843A9")
+#pragma message("UESHADERMETADATA_VERSION 8F516D4E-001A-4398-87CD-BB6D3450CF1D")
 
 #define COMPILER_SUPPORTS_ULONG_TYPES 1
 #define COMPILER_SUPPORTS_UINT64_IMAGE_ATOMICS (1 && COMPILER_SUPPORTS_ULONG_TYPES)
diff --git a/Engine/Source/Developer/ShaderCompilerCommon/Private/SpirvCommon.cpp b/Engine/Source/Developer/ShaderCompilerCommon/Private/SpirvCommon.cpp
--- a/Engine/Source/Developer/ShaderCompilerCommon/Private/SpirvCommon.cpp
+++ b/Engine/Source/Developer/ShaderCompilerCommon/Private/SpirvCommon.cpp
@@ -317,3 +317,131 @@ SHADERCOMPILERCOMMON_API void ParseSpirvGlobalVariables(const FSpirv& Spirv, Spv
 		}
 	}
 }
+
+// Scan the early sections of the SPIRV to find the used members of the specified structure
+SHADERCOMPILERCOMMON_API TArray<FSpirvStructureMemberData> ParseSpirvStructureMembers(const FSpirv& Spirv, FAnsiStringView StructureName)
+{
+	uint32_t StructureId = MAX_uint32;
+
+	TArray<FSpirvStructureMemberData> Members;
+
+	TMap<uint32, uint32> TypeSizes;
+
+	bool bReachedBody = false;
+	for (FSpirvConstIterator Iter = Spirv.begin(); Iter != Spirv.end() && !bReachedBody; ++Iter)
+	{
+		switch (Iter.Opcode())
+		{
+
+		// Use OpName to find the TypeId of the two structs we're interested in
+		case SpvOpName:
+		{
+			// Operands:
+			// 1 - Result Id
+			// 2 - Name
+
+			if (StructureId == MAX_uint32)
+			{
+				FAnsiStringView NameStr = FAnsiStringView(Iter.OperandAsString(2));
+				if (NameStr.EndsWith(StructureName))
+				{
+					StructureId = Iter.Operand(1);
+				}
+			}
+
+			break;
+		}
+
+		// Grab the OpMemberName for each member of the loose globals
+		case SpvOpMemberName:
+		{
+			// Operands:
+			// 1 - Type Id
+			// 2 - Member (index)
+			// 3 - Name
+
+			if (Iter.Operand(1) == StructureId)
+			{
+				const uint32_t Index = Iter.Operand(2);
+				if (Index >= (uint32)Members.Num())
+				{
+					Members.SetNum(Index + 1);
+				}
+				Members[Index].Name = FString(Iter.OperandAsString(3));
+			}
+			break;
+		}
+
+		// Take down the offset of each member (unused members aren't decorated)
+		case SpvOpMemberDecorate:
+		{
+			// Operands:
+			// 1 - Type Id
+			// 2 - Member (index)
+			// 3 - Decoration
+			// 4 - Decoration Value
+
+			if (Iter.Operand(1) == StructureId)
+			{
+				if (Iter.Operand(3) == SpvDecorationOffset)
+				{
+					const uint32_t Index = Iter.Operand(2);
+					Members[Index].Offset = Iter.Operand(4);
+				}
+			}
+			break;
+		}
+
+		case SpvOpTypeInt:
+		case SpvOpTypeFloat:
+		{
+			// Operands:
+			// 1 - Type Id
+			// 2 - Bit width
+			// 3 - Signedness (for ints)
+
+			TypeSizes.Add(Iter.Operand(1), Iter.Operand(2) / 8);
+			break;
+		}
+
+		case SpvOpTypeVector:
+		{
+			// Operands:
+			// 1 - Type Id
+			// 2 - Component Type Id
+			// 3 - Component Count
+
+			if (TypeSizes.Contains(Iter.Operand(1)))
+			{
+				const uint32 ComponentSize = TypeSizes[Iter.Operand(2)];
+				TypeSizes.Add(Iter.Operand(1), Iter.Operand(3) * ComponentSize);
+			}
+			break;
+		}
+
+		case SpvOpTypeStruct:
+		{
+			// Operands:
+			// 1 - Type Id
+			// 2,3,n - Member 0,1,n Type
+
+			if (Iter.Operand(1) == StructureId)
+			{
+				const uint32 OperandStartIndex = 2;
+				for (uint32 OperandIndex = OperandStartIndex; OperandIndex < Iter.WordCount(); ++OperandIndex)
+				{
+					Members[OperandIndex - OperandStartIndex].Size = TypeSizes[Iter.Operand(OperandIndex)];
+				}
+			}
+			break;
+		}
+
+		// When the body starts, that means there's no more useful information for us
+		case SpvOpFunction:
+			bReachedBody = true;
+			break;
+		}
+	}
+
+	return Members;
+}
\ No newline at end of file
diff --git a/Engine/Source/Developer/ShaderCompilerCommon/Public/SpirvCommon.h b/Engine/Source/Developer/ShaderCompilerCommon/Public/SpirvCommon.h
--- a/Engine/Source/Developer/ShaderCompilerCommon/Public/SpirvCommon.h
+++ b/Engine/Source/Developer/ShaderCompilerCommon/Public/SpirvCommon.h
@@ -203,3 +203,11 @@ SHADERCOMPILERCOMMON_API const ANSICHAR* PatchSpirvEntryPointWithCRC(FSpirv& Spi
 /** Parse global variables (anything but function local variables) of specified storage class from the specified SPIR-V module. */
 SHADERCOMPILERCOMMON_API void ParseSpirvGlobalVariables(const FSpirv& Spirv, SpvStorageClass StorageClass, TArray<FString>& OutVariableNames);
 
+/** Parse global variables (anything but function local variables) of specified storage class from the specified SPIR-V module.  Currently only supports plain data (integers, floats, vectors). */
+struct FSpirvStructureMemberData
+{
+	FString Name;
+	int32 Offset = -1;
+	uint32 Size = 0;
+};
+SHADERCOMPILERCOMMON_API TArray<FSpirvStructureMemberData> ParseSpirvStructureMembers(const FSpirv& Spirv, FAnsiStringView StructureName);
diff --git a/Engine/Source/Developer/VulkanShaderFormat/Private/SpirVShaderCompiler.inl b/Engine/Source/Developer/VulkanShaderFormat/Private/SpirVShaderCompiler.inl
--- a/Engine/Source/Developer/VulkanShaderFormat/Private/SpirVShaderCompiler.inl
+++ b/Engine/Source/Developer/VulkanShaderFormat/Private/SpirVShaderCompiler.inl
@@ -53,7 +53,7 @@ public:
 
 	TArray<FString> AllBindlessUBs;
 	TArray<FString> PushConstantUBs;
-	uint32 ShaderRecordGlobalsSize = 0;
+	bool bHasShaderRecordGlobals = false;
 
 	// Forwarded calls for convenience
 	inline EShaderFrequency GetShaderFrequency() const
@@ -693,6 +693,31 @@ static bool BuildShaderOutputFromSpirv(
 	FSpirvReflectBindings Bindings;
 	GatherSpirvReflectionBindings(Reflection, Bindings, SerializedOutput.UsedBindlessUB, InternalState);
 
+	if (InternalState.bHasShaderRecordGlobals)
+	{
+		// Scan the early sections of the SPIRV to find the full definition of the VulkanShaderRecordLooseParameterData in the FVulkanHitGroupSystemParameters
+		// so that we can properly reflect the offset/size of each of its members.  The compiler doesn't decorate unused members, so anything we find requires reflection.
+		TArray<FSpirvStructureMemberData> Members = ParseSpirvStructureMembers(SerializedOutput.Spirv, "VulkanShaderRecordLooseParameterData");
+		for (const FSpirvStructureMemberData& MemberData : Members)
+		{
+			if (MemberData.Name.Len() > 0)
+			{
+				checkf(MemberData.Offset >= 0, TEXT("Shader Record Global was named but never decorated with an Offset."));
+				checkf(MemberData.Size > 0, TEXT("Shader Record Global was named but the structure never assigned types."));
+
+				HandleReflectedGlobalConstantBufferMember(
+					InternalState.Input,
+					*InternalState.ParameterParser,
+					MemberData.Name,
+					0 /*ConstantBufferIndex*/,
+					MemberData.Offset /*ReflectionOffset*/,
+					MemberData.Size /*ReflectionSize*/,
+					Output
+				);
+			}
+		}
+	}
+
 	const FString UBOGlobalsNameSpv(CrossCompiler::FShaderConductorContext::GetIdentifierTable().GlobalsUniformBuffer);
 	const FString UBORootParamNameSpv(FShaderParametersMetadata::kRootUniformBufferBindingName);
 
@@ -1427,7 +1452,7 @@ static void ModifyCompilerInput(FSpirvShaderCompilerInternalState& InternalState
 	{
 		Input.Environment.SetDefine(TEXT("VULKAN_PROFILE_SM6"), 1);
 		Input.Environment.SetDefine(TEXT("PLATFORM_SUPPORTS_CALLABLE_SHADERS"), 1);
-//		Input.Environment.SetDefine(TEXT("PLATFORM_SUPPORTS_CONSTANTBUFFER_OBJECT"), true);
+		Input.Environment.SetDefine(TEXT("PLATFORM_SUPPORTS_CONSTANTBUFFER_OBJECT"), true);
 	}
 	else if (InternalState.IsSM5())
 	{
diff --git a/Engine/Source/Developer/VulkanShaderFormat/Private/VulkanShaderCompiler.cpp b/Engine/Source/Developer/VulkanShaderFormat/Private/VulkanShaderCompiler.cpp
--- a/Engine/Source/Developer/VulkanShaderFormat/Private/VulkanShaderCompiler.cpp
+++ b/Engine/Source/Developer/VulkanShaderFormat/Private/VulkanShaderCompiler.cpp
@@ -174,41 +174,23 @@ static uint32 GetSizeForType(FStringView TypeName, FStringView ArraySize)
 }
 
 
-// :todo-jn: TEMPORARY EXPERIMENT - will eventually move into preprocessing step
-static uint32 ConvertGlobalsToShaderRecord(const FShaderCompilerInput& Input, const FShaderParameterParser& ShaderParameterParser, const TMap<FStringView, FStringView>& ReplacedGlobals, FString& PreprocessedShaderSource, FShaderCompilerOutput& Output)
+// Add the globals as member in the VulkanShaderRecordLooseParameterData structure in the shader record
+static bool ConvertGlobalsToShaderRecord(const FShaderCompilerInput& Input, const FShaderParameterParser& ShaderParameterParser, const TMap<FStringView, FStringView>& ReplacedGlobals, FString& PreprocessedShaderSource, FShaderCompilerOutput& Output)
 {
-	uint32 ShaderRecordGlobalsSize = 0;
-	uint32 ShaderRecordParamCount = 0;
 	FString ShaderRecordGlobalsString;
-
 	for (const auto& ParamDecl : ReplacedGlobals)
 	{
 		ShaderRecordGlobalsString += ParamDecl.Value;
-
-		const FString ParamName(ParamDecl.Key);
-		const FShaderParameterParser::FParsedShaderParameter& Info = ShaderParameterParser.FindParameterInfos(ParamName);
-		const uint32 ParamSize = GetSizeForType(Info.ParsedType, Info.ParsedArraySize);
-
-		HandleReflectedGlobalConstantBufferMember(
-			Input,
-			ShaderParameterParser,
-			ParamName,
-			0 /*ConstantBufferIndex*/,
-			ShaderRecordGlobalsSize /*ReflectionOffset*/,
-			ParamSize /*ReflectionSize*/,
-			Output
-		);
-
-		ShaderRecordGlobalsSize += ParamSize; 
 	}
 
 	if (ShaderRecordGlobalsString.Len())
 	{
 		const int32 ReplacementCount = PreprocessedShaderSource.ReplaceInline(TEXT("uint VulkanShaderRecordDummyGlobals;"), *ShaderRecordGlobalsString, ESearchCase::CaseSensitive);
 		checkf(ReplacementCount == 1, TEXT("VulkanShaderRecordDummyGlobals was replaced %d times!"), ReplacementCount);
+		return true;
 	}
 
-	return ShaderRecordGlobalsSize;
+	return false;
 }
 
 
@@ -324,10 +306,10 @@ void CompileVulkanShader(const FShaderCompilerInput& Input, const FShaderPreproc
 	if (InternalState.bUseBindlessUniformBuffer)
 	{
 		InternalState.AllBindlessUBs = SpirvShaderCompiler::ConvertUBToBindless(PreprocessedSource, SpirvShaderCompiler::EBindlessUniformBufferType::ShaderRecord);
-		InternalState.ShaderRecordGlobalsSize = ConvertGlobalsToShaderRecord(Input, ShaderParameterParser, ReplacedGlobals, PreprocessedSource, Output);
+		InternalState.bHasShaderRecordGlobals = ConvertGlobalsToShaderRecord(Input, ShaderParameterParser, ReplacedGlobals, PreprocessedSource, Output);
 	}
 
-	if (ShaderParameterParser.DidModifyShader() || bSDCEMinified || InternalState.AllBindlessUBs.Num() || InternalState.ShaderRecordGlobalsSize || InternalState.PushConstantUBs.Num())
+	if (ShaderParameterParser.DidModifyShader() || bSDCEMinified || InternalState.AllBindlessUBs.Num() || InternalState.bHasShaderRecordGlobals || InternalState.PushConstantUBs.Num())
 	{
 		Output.ModifiedShaderSource = PreprocessedSource;
 	}
