Removing header deprecation.

#Synced-CL 50722762

diff --git a/Engine/Source/Developer/ShaderCompilerCommon/Public/CrossCompilerDefinitions.h b/Engine/Source/Developer/ShaderCompilerCommon/Public/CrossCompilerDefinitions.h
--- a/Engine/Source/Developer/ShaderCompilerCommon/Public/CrossCompilerDefinitions.h
+++ b/Engine/Source/Developer/ShaderCompilerCommon/Public/CrossCompilerDefinitions.h
@@ -4,8 +4,6 @@
 
 #include "Misc/CoreMiscDefines.h"
 
-UE_DEPRECATED_HEADER(5.8, "CrossCompilerDefinitions.h is no longer used")
-
 //@todo-lh: As soon as HLSLcc is fully replaced by DXC for all shader compiler backends,
 // move the content of this file into 'CrossCompilerDefinitions.h".
 // Until then, we simply include that file from the HLSLcc folder since it only contains enum declaration but no further include dependencies.
