SDCE, fixed crash on parsing Vulkan annotations
- Postfix declaration annotations [[vk::...]] were incorrectly parsed as an array.

#Synced-CL 50766172

diff --git a/Engine/Source/Developer/ShaderCompilerCommon/Private/SDCE/SDCEShallowParser.h b/Engine/Source/Developer/ShaderCompilerCommon/Private/SDCE/SDCEShallowParser.h
--- a/Engine/Source/Developer/ShaderCompilerCommon/Private/SDCE/SDCEShallowParser.h
+++ b/Engine/Source/Developer/ShaderCompilerCommon/Private/SDCE/SDCEShallowParser.h
@@ -739,12 +739,22 @@ namespace UE::ShaderMinifier::SDCE
 			return true;
 		}
 
-		void TryParseTypePostfix(FShallowTokenizer& Tk)
+		void TryParseTypePostfix(FShallowTokenizer& TkOuter)
 		{
+			FShallowTokenizer Tk = Branch(TkOuter);
+			
 			for (;;)
 			{
 				if (Tk.IsConsume(ETokenType::SquareOpen))
 				{
+					// Reject annotations
+					if (Tk.Is(ETokenType::SquareOpen))
+					{
+						// No longer a type postfix, any prior postfixes have been accepted at this point
+						Reject(Tk);
+						break;
+					}
+					
 					// May close immediately for automatic sizing
 					if (!Tk.Is(ETokenType::SquareClose))
 					{
@@ -755,10 +765,14 @@ namespace UE::ShaderMinifier::SDCE
 					
 					check(Tk.Is(ETokenType::SquareClose));
 					Tk.Next();
+					
+					Accept(TkOuter, Tk);
+					Tk = Branch(TkOuter);
 					continue;
 				}
-
+				
 				// Not a postfix
+				Reject(Tk);
 				break;
 			}
 		}
diff --git a/Engine/Source/Developer/ShaderCompilerCommon/Private/SDCE/SDCETokenizer.h b/Engine/Source/Developer/ShaderCompilerCommon/Private/SDCE/SDCETokenizer.h
--- a/Engine/Source/Developer/ShaderCompilerCommon/Private/SDCE/SDCETokenizer.h
+++ b/Engine/Source/Developer/ShaderCompilerCommon/Private/SDCE/SDCETokenizer.h
@@ -154,6 +154,19 @@ namespace UE::ShaderMinifier::SDCE
 			// If there's a valid op head, this is a branch that hasn't been accepted/rejected
 			check(OpHead == TokenizerInvalidOpHead);
 		}
+		
+		FShallowTokenizer& operator=(FShallowTokenizer&& rhs)
+		{
+			checkf(OpHead == TokenizerInvalidOpHead, TEXT("Assignment to a pending tokenizer"));
+			Token = rhs.Token;
+			OpHead = rhs.OpHead;
+			Begin = rhs.Begin;
+			End = rhs.End;
+			
+			// Invalidate RHS op-head
+			rhs.OpHead = TokenizerInvalidOpHead;
+			return *this;
+		}
 
 		/** Branch this tokenizer */
 		FShallowTokenizer Branch() const
