RHI: Add bindless resource support to specialization constant exclusions.

#rnx

#Synced-CL 51094475

diff --git a/Engine/Source/Developer/ShaderCompilerCommon/Internal/ShaderCompilerCommonInternal.cpp b/Engine/Source/Developer/ShaderCompilerCommon/Internal/ShaderCompilerCommonInternal.cpp
--- a/Engine/Source/Developer/ShaderCompilerCommon/Internal/ShaderCompilerCommonInternal.cpp
+++ b/Engine/Source/Developer/ShaderCompilerCommon/Internal/ShaderCompilerCommonInternal.cpp
@@ -2,6 +2,7 @@
 
 #include "ShaderCompilerCommonInternal.h"
 #include "HlslParserInternal.h"
+#include "ShaderParameterParser.h"
 
 
 namespace UE::ShaderCompilerCommon
@@ -236,6 +237,8 @@ int32 ProcessSpirvSpecializationConstants(
 		"--cfg-cleanup"
 	};
 
+	const bool bIsBindless = Input.IsBindlessEnabled();
+
 	// Map of resource names to hashes of spec const values to exclude the resource from
 	TMap<FString, TArray<uint32>> ExcludedResources;
 	auto DiffResourceArrays = [&](const uint32 SpecHash, const TArray<SpvReflectDescriptorBinding*>& UnspecArray, const TArray<SpvReflectDescriptorBinding*>& SpecArray)
@@ -252,6 +255,28 @@ int32 ProcessSpirvSpecializationConstants(
 			}
 		};
 
+	auto DiffBindlessResources = [&](const uint32 SpecHash, const SpvReflectBlockVariable& UnspecGlobals, const SpvReflectBlockVariable& SpecGlobals)
+		{
+			TConstArrayView<SpvReflectBlockVariable> UnspecMembers(UnspecGlobals.members, UnspecGlobals.member_count);
+			TConstArrayView<SpvReflectBlockVariable> SpecMembers(SpecGlobals.members, SpecGlobals.member_count);
+			for (const SpvReflectBlockVariable& UnspecMember : UnspecMembers)
+			{
+				const FString MemberName(UnspecMember.name);
+				FStringView AdjustedMemberName(MemberName);
+				const EShaderParameterType BindlessParameterType = FShaderParameterParser::ParseAndRemoveBindlessParameterPrefix(AdjustedMemberName);
+				if (BindlessParameterType != EShaderParameterType::LooseData)
+				{
+					if (!SpecMembers.FindByPredicate([&UnspecMember](const SpvReflectBlockVariable& SpecMember) {
+						return (FCStringAnsi::Strcmp(UnspecMember.name, SpecMember.name) == 0);
+						}))
+					{
+						TArray<uint32>& Hashes = ExcludedResources.FindOrAdd(FString(AdjustedMemberName));
+						Hashes.Add(SpecHash);
+					}
+				}
+			}
+		};
+
 	auto ProcessFullSpecialization = [&](const TArray<int32>& SpecValues)
 		{
 			// Set the default values to the combination we are evaluating
@@ -292,15 +317,36 @@ int32 ProcessSpirvSpecializationConstants(
 				EvalBindings.GatherDescriptorBindings(EvalReflection);
 
 				const uint32 SpecHash = GetTypeHash(SpecValues);
-				DiffResourceArrays(SpecHash, UnspecializedBindings.UniformBuffers, EvalBindings.UniformBuffers);
-				DiffResourceArrays(SpecHash, UnspecializedBindings.Samplers, EvalBindings.Samplers);
-				DiffResourceArrays(SpecHash, UnspecializedBindings.TextureSRVs, EvalBindings.TextureSRVs);
-				DiffResourceArrays(SpecHash, UnspecializedBindings.TextureUAVs, EvalBindings.TextureUAVs);
-				DiffResourceArrays(SpecHash, UnspecializedBindings.TBufferSRVs, EvalBindings.TBufferSRVs);
-				DiffResourceArrays(SpecHash, UnspecializedBindings.TBufferUAVs, EvalBindings.TBufferUAVs);
-				DiffResourceArrays(SpecHash, UnspecializedBindings.SBufferSRVs, EvalBindings.SBufferSRVs);
-				DiffResourceArrays(SpecHash, UnspecializedBindings.SBufferUAVs, EvalBindings.SBufferUAVs);
-				DiffResourceArrays(SpecHash, UnspecializedBindings.AccelerationStructures, EvalBindings.AccelerationStructures);
+
+				if (bIsBindless)
+				{
+					const SpvReflectDescriptorBinding* const* UnspecializedGlobals = 
+						UnspecializedBindings.UniformBuffers.FindByPredicate([](const SpvReflectDescriptorBinding* Binding) {
+							return FCStringAnsi::Strcmp("$Globals", Binding->name) == 0;
+						});
+
+					SpvReflectDescriptorBinding** SpecializedGlobals =
+						EvalBindings.UniformBuffers.FindByPredicate([](SpvReflectDescriptorBinding* Binding) {
+							return FCStringAnsi::Strcmp("$Globals", Binding->name) == 0;
+						});
+
+					if (UnspecializedGlobals && SpecializedGlobals)
+					{
+						DiffBindlessResources(SpecHash, (*UnspecializedGlobals)->block, (*SpecializedGlobals)->block);
+					}
+				}
+				else
+				{
+					DiffResourceArrays(SpecHash, UnspecializedBindings.UniformBuffers, EvalBindings.UniformBuffers);
+					DiffResourceArrays(SpecHash, UnspecializedBindings.Samplers, EvalBindings.Samplers);
+					DiffResourceArrays(SpecHash, UnspecializedBindings.TextureSRVs, EvalBindings.TextureSRVs);
+					DiffResourceArrays(SpecHash, UnspecializedBindings.TextureUAVs, EvalBindings.TextureUAVs);
+					DiffResourceArrays(SpecHash, UnspecializedBindings.TBufferSRVs, EvalBindings.TBufferSRVs);
+					DiffResourceArrays(SpecHash, UnspecializedBindings.TBufferUAVs, EvalBindings.TBufferUAVs);
+					DiffResourceArrays(SpecHash, UnspecializedBindings.SBufferSRVs, EvalBindings.SBufferSRVs);
+					DiffResourceArrays(SpecHash, UnspecializedBindings.SBufferUAVs, EvalBindings.SBufferUAVs);
+					DiffResourceArrays(SpecHash, UnspecializedBindings.AccelerationStructures, EvalBindings.AccelerationStructures);
+				}
 			}
 		};
 
diff --git a/Engine/Source/Runtime/RenderCore/Private/ShaderParameterStruct.cpp b/Engine/Source/Runtime/RenderCore/Private/ShaderParameterStruct.cpp
--- a/Engine/Source/Runtime/RenderCore/Private/ShaderParameterStruct.cpp
+++ b/Engine/Source/Runtime/RenderCore/Private/ShaderParameterStruct.cpp
@@ -628,6 +628,7 @@ FResourceExclusionArray GetSpecializationResourceExclusions(const TShaderRef<FSh
 
 		const uint32 SpecializationHash = GetTypeHash(Shader.GetSpecializationConstants());
 		const FShaderParameterBindings& Bindings = Shader->Bindings;
+		// Only one of the two following loops will fill ResourceExclusions (depending on if bindless is enabled)
 		for (const FShaderParameterBindings::FResourceParameter& Parameter : Bindings.ResourceParameters)
 		{
 			// Check if this resource should be excluded from this combination of specialization constants
@@ -636,6 +637,14 @@ FResourceExclusionArray GetSpecializationResourceExclusions(const TShaderRef<FSh
 				ResourceExclusions.Add(Parameter.BaseIndex);
 			}
 		}
+		for (const FShaderParameterBindings::FBindlessResourceParameter& Parameter : Bindings.BindlessResourceParameters)
+		{
+			// Check if this resource should be excluded from this combination of specialization constants
+			if (RHIShader->IsExcludedFromSpecialization(Parameter.GlobalConstantOffset, SpecializationHash))
+			{
+				ResourceExclusions.Add(Parameter.GlobalConstantOffset);
+			}
+		}
 
 		return ResourceExclusions;
 	}
@@ -877,10 +886,15 @@ static void ExtractShaderParameterResources(
 	const FShaderParameterReader Reader(ParametersData);
 
 #if PLATFORM_SUPPORTS_BINDLESS_RENDERING
-	if (int32 NumBindings = Bindings.BindlessResourceParameters.Num())
+	if (Bindings.BindlessResourceParameters.Num())
 	{
 		for (const FShaderParameterBindings::FBindlessResourceParameter& Parameter : Bindings.BindlessResourceParameters)
 		{
+			if (ResourceExclusions.Contains(Parameter.GlobalConstantOffset))
+			{
+				continue;
+			}
+
 			const FRHIShaderParameterResource ShaderParameterResource = ExtractShaderParameterResource(Reader, Parameter);
 			BatchedParameters.AddBindlessParameter(ShaderParameterResource);
 		}
