Move shader symbol export to a background task, allowing the game thread to continue work. Saves ~6 seconds on editor startup if shader symbol export is on.
Symbols are still written out in a serial fashion by having each write-out task depend on the previous task, avoiding race conditions.

#Synced-CL 52172106

diff --git a/Engine/Source/Developer/ShaderCompilerCommon/Private/ShaderSymbolExport.cpp b/Engine/Source/Developer/ShaderCompilerCommon/Private/ShaderSymbolExport.cpp
--- a/Engine/Source/Developer/ShaderCompilerCommon/Private/ShaderSymbolExport.cpp
+++ b/Engine/Source/Developer/ShaderCompilerCommon/Private/ShaderSymbolExport.cpp
@@ -147,6 +147,8 @@ void FShaderSymbolExport::Initialize()
 
 void FShaderSymbolExport::WriteSymbolData(const FString& Filename, const FString& DebugData, TConstArrayView<uint8> Contents)
 {
+	TRACE_CPUPROFILER_EVENT_SCOPE(FShaderSymbolExport::WriteSymbolData);
+
 	// No writing is possible if the Filename is empty
 	if (Filename.IsEmpty())
 	{
@@ -159,102 +161,163 @@ void FShaderSymbolExport::WriteSymbolData(const FString& Filename, const FString
 	// Otherwise, we add an additional condition based on the hash of the symbols, if the new hash is less than the one previously
 	// written we re-write (so we always export a deterministic copy of the symbols in this case).
 
-	bool bAlreadyInSet = false;
-	ExportedShaders.Add(Filename, &bAlreadyInSet);
 	bool bDeterministicSymbols = ZipWriter == nullptr && !ShouldGenerateShaderSymbolsInfo(ShaderFormat);
 
 	// calculate the hash of the new symbols only if we're not exporting deterministically
 	// (we don't need it in that case as described in the comments above)
 	uint64 NewSymbolHash = bDeterministicSymbols ? CityHash64(reinterpret_cast<const char*>(Contents.GetData()), Contents.Num()) : 0u;
-	if (bAlreadyInSet)
+
+	// Read phase: early-out for shaders already exported, without acquiring the write lock.
 	{
-		// We've already exported a copy of symbols for this shader hash and are not ensuring exports are deterministic
-		DuplicateSymbols++;
-		if (!bDeterministicSymbols)
+		FRWScopeLock ReadLock(SymbolWriteLock, SLT_ReadOnly);
+		if (ExportedShaders.Contains(Filename))
 		{
-			return;
+			if (!bDeterministicSymbols)
+			{
+				// DuplicateSymbols is only incremented in the read phase on early return, because the write phase will also increment.
+				// Remember that this block may fail to detect a duplicate if two threads race to insert.
+				DuplicateSymbols.fetch_add(1, std::memory_order_relaxed);
+				return;
+			}
+
+			// when ensuring determinism, we further store the hash of the symbols we wrote out, and if the new hash is less than
+			// the previously written one we write it again, so we always choose the same copy of symbols to write during i.e. a cook
+			const FSymbolFileInfo* PrevExportedSymbolInfo = ExportedSymbolInfo.Find(Filename);
+			check(PrevExportedSymbolInfo);
+			check(NewSymbolHash != 0u);
+			if (NewSymbolHash >= PrevExportedSymbolInfo->Hash)
+			{
+				DuplicateSymbols.fetch_add(1, std::memory_order_relaxed);
+				return;
+			}
+			// Hash is lower; fall through to write phase to overwrite
 		}
-		
-		// when ensuring determinism, we further store the hash of the symbols we wrote out, and if the new hash is less than
-		// the previously written one we write it again, so we always choose the same copy of symbols to write during i.e. a cook
-		FSymbolFileInfo* PrevExportedSymbolInfo = ExportedSymbolInfo.Find(Filename);
-		
-		check(PrevExportedSymbolInfo);
-		check(NewSymbolHash != 0u);
-		if (NewSymbolHash >= PrevExportedSymbolInfo->Hash)
+	}
+
+	// Write phase: first insertion, or deterministic overwrite with a lower hash.
+	// Re-check after acquiring the write lock to handle races on first insertion.
+	bool bShouldWriteSymbols = ShouldWriteShaderSymbols(ShaderFormat) && Contents.Num();
+	bool bShouldWriteToZip = ZipWriter != nullptr;
+	{
+		FRWScopeLock WriteLock(SymbolWriteLock, SLT_Write);
+
+		bool bAlreadyInSet = false; // Two threads may be trying to insert concurrently
+		ExportedShaders.Add(Filename, &bAlreadyInSet);
+
+		if (bAlreadyInSet)
 		{
-			return;
+			// We've already exported a copy of symbols for this shader hash and are not ensuring exports are deterministic
+			DuplicateSymbols.fetch_add(1, std::memory_order_relaxed);
+			if (!bDeterministicSymbols)
+			{
+				return;
+			}
+
+			FSymbolFileInfo* PrevExportedSymbolInfo = ExportedSymbolInfo.Find(Filename);
+			check(PrevExportedSymbolInfo);
+			check(NewSymbolHash != 0u);
+			if (NewSymbolHash >= PrevExportedSymbolInfo->Hash)
+			{
+				return;
+			}
+
+			// new hash is less; adjust stats to account for file we will replace below
+			TotalSymbolDataBytes -= PrevExportedSymbolInfo->Size;
+			TotalSymbolData--;
 		}
 
-		// new hash is less; adjust stats to account for file we will replace below
-		TotalSymbolDataBytes -= PrevExportedSymbolInfo->Size;
-		TotalSymbolData--;
-	}
+		if (bDeterministicSymbols)
+		{
+			ExportedSymbolInfo.Add(Filename, { NewSymbolHash, Contents.Num() });
+		}
 
-	if (bDeterministicSymbols)
-	{
-		ExportedSymbolInfo.Add(Filename, { NewSymbolHash, Contents.Num() });
-	}
+		// Emit periodic log messages detailing the size of the shader symbols output file/directory.
+		static uint64 LastReport = 0;
+		TotalSymbolDataBytes += Contents.Num();
+		TotalSymbolData++;
+
+		if ((TotalSymbolDataBytes - LastReport) >= (64 * 1024 * 1024))
+		{
+			UE_LOG(LogShaderSymbolExport, Display, TEXT("Shader symbols export size: %.2f MB, count: %llu"),
+				double(TotalSymbolDataBytes) / (1024.0 * 1024.0), TotalSymbolData);
+			LastReport = TotalSymbolDataBytes;
+		}
 
-	// Emit periodic log messages detailing the size of the shader symbols output file/directory.
-	static uint64 LastReport = 0;
-	TotalSymbolDataBytes += Contents.Num();
-	TotalSymbolData++;
+		if (ShouldGenerateShaderSymbolsInfo(ShaderFormat) && !DebugData.IsEmpty())
+		{
+			// Collect the simple shader symbol information
+			FString FilenameNoExt = FPaths::GetBaseFilename(Filename);
+			ShaderInfos.Add({ FilenameNoExt, DebugData });
+		}
 
-	if ((TotalSymbolDataBytes - LastReport) >= (64 * 1024 * 1024))
+		if (bShouldWriteSymbols && bShouldWriteToZip)
+		{
+			// Append the platform data to the zip file
+			ZipWriter->AddFile(Filename, Contents, FDateTime::Now());
+		}
+	}
+	// Outside of lock scope
+	if (bShouldWriteSymbols && !bShouldWriteToZip)
 	{
-		UE_LOG(LogShaderSymbolExport, Display, TEXT("Shader symbols export size: %.2f MB, count: %llu"),
-			double(TotalSymbolDataBytes) / (1024.0 * 1024.0), TotalSymbolData);
-		LastReport = TotalSymbolDataBytes;
+		// In deterministic mode, files may be overwritten.
+		// If a thread wants to overwrite a file while the previous is still being written,
+		// we need to wait for the first one to finish.
+		WriteFile(Filename, Contents, bDeterministicSymbols);
 	}
+}
 
-	if (ShouldGenerateShaderSymbolsInfo(ShaderFormat) && !DebugData.IsEmpty())
+void FShaderSymbolExport::WriteFile(const FString& Filename, const TConstArrayView<uint8>& Contents, bool bAllowRetry)
+{
+	TRACE_CPUPROFILER_EVENT_SCOPE(FShaderSymbolExport::WriteFile);
+	IPlatformFile& PlatformFile = FPlatformFileManager::Get().GetPlatformFile();
+
+	// Write the symbols to the export directory
+	const FString OutputPath = ExportPath / Filename;
+	const FString Directory = FPaths::GetPath(OutputPath);
+
+	// Filename could contain extra folders, so we need to make sure they exist first.
+	if (!PlatformFile.CreateDirectoryTree(*Directory))
 	{
-		// Collect the simple shader symbol information
-		FString FilenameNoExt = FPaths::GetBaseFilename(Filename);
-		ShaderInfos.Add({ FilenameNoExt, DebugData });
+		UE_LOG(LogShaderSymbolExport, Error, TEXT("Failed to create shader symbol directory \"%s\"."), *Directory);
 	}
-
-	if (ShouldWriteShaderSymbols(ShaderFormat) && Contents.Num())
+	else
 	{
-		if (ZipWriter)
+		FFileOpenResult FileResult = PlatformFile.OpenWrite(*OutputPath, IPlatformFile::EOpenWriteFlags::None);
+
+		const int MaxRetries = 10;
+		int Retries = 0;
+		while (FileResult.HasError() && bAllowRetry && (Retries++ < MaxRetries))
 		{
-			// Append the platform data to the zip file
-			ZipWriter->AddFile(Filename, Contents, FDateTime::Now());
+			FPlatformProcess::Sleep(0.0f);
+			FRWScopeLock WriteLock(SymbolWriteLock, SLT_Write);
+			FileResult = PlatformFile.OpenWrite(*OutputPath, IPlatformFile::EOpenWriteFlags::None);
 		}
-		else
-		{
-			IPlatformFile& PlatformFile = FPlatformFileManager::Get().GetPlatformFile();
-
-			// Write the symbols to the export directory
-			const FString OutputPath = ExportPath / Filename;
-			const FString Directory = FPaths::GetPath(OutputPath);
 
-			// Filename could contain extra folders, so we need to make sure they exist first.
-			if (!PlatformFile.CreateDirectoryTree(*Directory))
-			{
-				UE_LOG(LogShaderSymbolExport, Error, TEXT("Failed to create shader symbol directory \"%s\"."), *Directory);
-			}
-			else
-			{
-				IFileHandle* File = PlatformFile.OpenWrite(*OutputPath);
-				if (!File || !File->Write(Contents.GetData(), Contents.Num()))
-				{
-					UE_LOG(LogShaderSymbolExport, Error, TEXT("Failed to export shader symbols \"%s\"."), *OutputPath);
-				}
+		if (FileResult.HasError())
+		{
+			FFileSystemError& Error = FileResult.GetError();
+			FString ErrorMessage = Error.GetMessage();
+			UE_LOG(LogShaderSymbolExport, Error, TEXT("Failed to open shader symbols file \"%s\":"), *OutputPath);
+			UE_LOG(LogShaderSymbolExport, Error, TEXT("%s"), *ErrorMessage);
+			return;
+		}
 
-				if (File)
-				{
-					delete File;
-				}
-			}
+		TUniquePtr<IFileHandle> File = FileResult.StealValue();
+		if (!File->Write(Contents.GetData(), Contents.Num()))
+		{
+			UE_LOG(LogShaderSymbolExport, Error, TEXT("Failed to write shader symbols \"%s\"."), *OutputPath);
 		}
 	}
 }
 
 void FShaderSymbolExport::NotifyShaderCompilersShutdown()
 {
-	FScopeLock Lock(&Cs);
+	// Wait for all in-flight write tasks to complete before finalizing artifacts.
+	AllTasksDoneEvent->Reset();
+	if (PendingTaskCount.load(std::memory_order_acquire) > 0)
+	{
+		AllTasksDoneEvent->Wait();
+	}
 
 	// Used to match worker files, for example (if not overridden): [ShaderSymbols_]N.ext
 	const FString WorkerName = ExportFileName + TEXT("_");
@@ -339,9 +402,10 @@ void FShaderSymbolExport::NotifyShaderCompilersShutdown()
 				Writer->Serialize(Output.GetData(), Output.Num());
 				Writer->Close();
 				UE_LOG(LogShaderSymbolExport, Display, TEXT("Wrote %d records into shader symbols info output file \"%s\"."), ShaderInfos.Num(), *InfoFilePath);
-				const uint32 TotalShaders = ShaderInfos.Num() + DuplicateSymbols;
+				const uint32 NumDuplicates = DuplicateSymbols.load(std::memory_order_relaxed);
+				const uint32 TotalShaders = ShaderInfos.Num() + NumDuplicates;
 				UE_LOG(LogShaderSymbolExport, Display, TEXT("%d total shaders, %d shaders after platform deduplication. %d duplicates (%4.1f%%)."),
-					TotalShaders, TotalShaders-DuplicateSymbols, DuplicateSymbols, (float)DuplicateSymbols / TotalShaders * 100.0f);
+					TotalShaders, TotalShaders - NumDuplicates, NumDuplicates, (float)NumDuplicates / TotalShaders * 100.0f);
 			}
 			else
 			{
diff --git a/Engine/Source/Developer/ShaderCompilerCommon/Public/ShaderSymbolExport.h b/Engine/Source/Developer/ShaderCompilerCommon/Public/ShaderSymbolExport.h
--- a/Engine/Source/Developer/ShaderCompilerCommon/Public/ShaderSymbolExport.h
+++ b/Engine/Source/Developer/ShaderCompilerCommon/Public/ShaderSymbolExport.h
@@ -6,16 +6,21 @@
 
 #if WITH_ENGINE
 
+#include "Tasks/Task.h"
 #include "Compression/CompressedBuffer.h"
 #include "CoreTypes.h"
 #include "Containers/ContainersFwd.h"
 #include "Containers/Map.h"
 #include "Containers/Set.h"
 #include "HAL/CriticalSection.h"
+#include "HAL/Event.h"
+#include "Misc/ScopeRWLock.h"
 #include "Misc/ScopeLock.h"
 #include "Serialization/MemoryReader.h"
 #include "Templates/UniquePtr.h"
 
+#include <atomic>
+
 class FZipArchiveWriter;
 
 class FShaderSymbolExport
@@ -41,6 +46,7 @@ public:
 private:
 	SHADERCOMPILERCOMMON_API void Initialize();
 	SHADERCOMPILERCOMMON_API void WriteSymbolData(const FString& Filename, const FString& DebugInfo, TConstArrayView<uint8> Contents);
+	void WriteFile(const FString& Filename, const TConstArrayView<uint8>& Contents, bool bAllowRetry);
 
 	const FName ShaderFormat;
 
@@ -55,15 +61,13 @@ private:
 	FString ExportPath;
 	FString InfoFilePath;
 	FString ExportFileName;
-	FCriticalSection Cs;
 	uint64 TotalSymbolDataBytes{ 0 };
 	uint64 TotalSymbolData{ 0 };
-	bool bInitialized{ false };
 	bool bExportShaderSymbols{ false };
 
 	TMap<FString, FString> ShaderInfos;
 
-	uint32 DuplicateSymbols{ 0 };
+	std::atomic<uint32> DuplicateSymbols{ 0 };
 
 	/**
 	 * If true, the current process is the first process in a multiprocess group, or is not in a group,
@@ -71,36 +75,63 @@ private:
 	 * is necessary for given settings.
 	 */
 	bool bMultiprocessOwner{ false };
+
+	std::atomic<bool> bInitialized{ false };
+	FCriticalSection InitCs;
+
+	// Write task synchronization: tasks run concurrently for decompress/deserialize;
+	// PendingTaskCount tracks in-flight tasks; AllTasksDoneEvent is triggered when it reaches
+	// zero so NotifyShaderCompilersShutdown can drain without holding references to each task.
+	FRWLock SymbolWriteLock;
+	std::atomic<int32> PendingTaskCount{ 0 };
+	FEventRef AllTasksDoneEvent{ EEventMode::ManualReset };
 };
 
 template<typename TPlatformShaderSymbolData>
 inline void FShaderSymbolExport::NotifyShaderCompiled(const FCompressedBuffer& PlatformSymbolDataCompressed, const FString& DebugInfo)
 {
-	FScopeLock Lock(&Cs);
-	if (!bInitialized)
+	// Double-checked locking: atomic acquire load on the fast path, InitCs taken only on first call
+	if (!bInitialized.load(std::memory_order_acquire))
 	{
-		// If we get called, we know we're compiling. Do one time initialization
-		// which will create the output directory / open the open file stream.
-		Initialize();
-		bInitialized = true;
+		FScopeLock Lock(&InitCs);
+		if (!bInitialized.load(std::memory_order_relaxed))
+		{
+			// If we get called, we know we're compiling. Do one time initialization
+			// which will create the output directory / open the output file stream.
+			Initialize();
+			bInitialized.store(true, std::memory_order_release);
+		}
 	}
 
-	if (bExportShaderSymbols)
+	if (!bExportShaderSymbols) // read-only after init; safe via acquire above
 	{
-		// Deserialize the platform symbol data
-		TPlatformShaderSymbolData FullSymbolData;
-		FSharedBuffer PlatformSymbolData = PlatformSymbolDataCompressed.Decompress();
-		FMemoryReaderView Ar(PlatformSymbolData.GetView());
-		Ar << FullSymbolData;
+		return;
+	}
 
-		for (const auto& SymbolData : FullSymbolData.GetAllSymbolData())
+	// Dispatch a task to decompress, deserialize, and write. All heavy work off the calling thread.
+	// The counter is incremented before launch; the task decrements it on completion and triggers
+	// AllTasksDoneEvent when it reaches zero so shutdown can drain without holding task references.
+	PendingTaskCount.fetch_add(1, std::memory_order_relaxed);
+	
+	UE::Tasks::Launch(TEXT("ShaderSymbolExport::WriteSymbolData"),
+		[this, PlatformSymbolDataCompressed, DebugInfo]()
 		{
-			const FString FileName = SymbolData.GetFilename();
-			TConstArrayView<uint8> Contents = SymbolData.GetContents();
-
-			WriteSymbolData(FileName, DebugInfo, Contents);
-		}
-	}
+			TPlatformShaderSymbolData FullSymbolData;
+			FSharedBuffer PlatformSymbolData = PlatformSymbolDataCompressed.Decompress();
+			FMemoryReaderView Ar(PlatformSymbolData.GetView());
+			Ar << FullSymbolData;
+
+			for (const auto& SymbolData : FullSymbolData.GetAllSymbolData())
+			{
+				WriteSymbolData(SymbolData.GetFilename(), DebugInfo, SymbolData.GetContents());
+			}
+
+			if (PendingTaskCount.fetch_sub(1, std::memory_order_acq_rel) == 1)
+			{
+				AllTasksDoneEvent->Trigger();
+			}
+		},
+		UE::Tasks::ETaskPriority::BackgroundLow);
 }
 
 #endif // WITH_ENGINE
