RHI: Move the exclusions for specialization constants out of optional data and into the vulkan header.  Add the possible specialization value array to the vulkan header (will be needed for shader objects).

#rnx

#Synced-CL 52545175

diff --git a/Engine/Shaders/Public/Platform/Vulkan/VulkanCommon.ush b/Engine/Shaders/Public/Platform/Vulkan/VulkanCommon.ush
--- a/Engine/Shaders/Public/Platform/Vulkan/VulkanCommon.ush
+++ b/Engine/Shaders/Public/Platform/Vulkan/VulkanCommon.ush
@@ -7,7 +7,7 @@
 #pragma once
 
 // Update this GUID to invalidate shader recompilation for only Vulkan shaders
-#pragma message("UESHADERMETADATA_VERSION 62D5B8B1-DECD-4C2D-98A4-6E5FE07DEA2E")
+#pragma message("UESHADERMETADATA_VERSION 2252475E-CF8A-45F3-85D7-E06B3EDCCCFC")
 
 #define COMPILER_SUPPORTS_ULONG_TYPES 1
 #define COMPILER_SUPPORTS_UINT64_IMAGE_ATOMICS (1 && COMPILER_SUPPORTS_ULONG_TYPES)
diff --git a/Engine/Source/Developer/ShaderCompilerCommon/Internal/ShaderCompilerCommonInternal.cpp b/Engine/Source/Developer/ShaderCompilerCommon/Internal/ShaderCompilerCommonInternal.cpp
--- a/Engine/Source/Developer/ShaderCompilerCommon/Internal/ShaderCompilerCommonInternal.cpp
+++ b/Engine/Source/Developer/ShaderCompilerCommon/Internal/ShaderCompilerCommonInternal.cpp
@@ -139,7 +139,7 @@ IMalloc* GetDxcMalloc()
 
 // Go through the specialization constants in the SPIRV to match them to what was expected in the build env.
 // Evaluate all the different combinations of the spec consts to find the resources that are conditional on specific values.
-int32 ProcessSpirvSpecializationConstants(
+FSpecializationConstantData ProcessSpirvSpecializationConstants(
 	const FShaderCompilerInput& Input, 
 	CrossCompiler::FShaderConductorContext& CompilerContext,
 	spv_reflect::ShaderModule& Reflection,
@@ -152,7 +152,7 @@ int32 ProcessSpirvSpecializationConstants(
 	// Bail if specialization constants aren't used
 	if (TotalSpecializationCount <= 0)
 	{
-		return 0;
+		return FSpecializationConstantData();
 	}
 
 	// Put the total count of spec constants
@@ -381,17 +381,22 @@ int32 ProcessSpirvSpecializationConstants(
 	//   Run GatherSpirvReflectionBindings and compare the resources
 	//     Mark the unused ones in the Parameter map
 
-	auto& ParameterMap = Output.ParameterMap.ParameterMap;
+	FSpecializationConstantData SpecializationConstantData;
+	SpecializationConstantData.Values = Input.Environment.SpecializationConstantValues;
+
+	TMap<uint32, TArray<uint32>> ExclusionMap;
+	const TMap<FString, FParameterAllocation>& ParameterMap = Output.ParameterMap.ParameterMap;
 	for (auto& Pair : ExcludedResources)
 	{
-		if (FParameterAllocation* Param = ParameterMap.Find(Pair.Key))
+		if (const FParameterAllocation* Param = ParameterMap.Find(Pair.Key))
 		{
-			TArray<uint32>& Hashes = Output.SpecializationConstantExclusions.FindOrAdd(Param->BaseIndex);
+			TArray<uint32>& Hashes = ExclusionMap.FindOrAdd(Param->BaseIndex);
 			Hashes = MoveTemp(Pair.Value);
 		}
 	}
-
-	return TotalSpecializationCount;
+	SpecializationConstantData.SetExclusionMap(MoveTemp(ExclusionMap));
+	SpecializationConstantData.Sort();
+	return SpecializationConstantData;
 }
 
 } // UE::ShaderCompilerCommon
diff --git a/Engine/Source/Developer/ShaderCompilerCommon/Internal/ShaderCompilerCommonInternal.h b/Engine/Source/Developer/ShaderCompilerCommon/Internal/ShaderCompilerCommonInternal.h
--- a/Engine/Source/Developer/ShaderCompilerCommon/Internal/ShaderCompilerCommonInternal.h
+++ b/Engine/Source/Developer/ShaderCompilerCommon/Internal/ShaderCompilerCommonInternal.h
@@ -41,7 +41,7 @@ namespace UE::ShaderCompilerCommon
 	extern SHADERCOMPILERCOMMON_API IMalloc* GetDxcMalloc();
 #endif
 
-	extern SHADERCOMPILERCOMMON_API int32 ProcessSpirvSpecializationConstants(
+	extern SHADERCOMPILERCOMMON_API FSpecializationConstantData ProcessSpirvSpecializationConstants(
 		const FShaderCompilerInput& Input,
 		CrossCompiler::FShaderConductorContext& CompilerContext,
 		spv_reflect::ShaderModule& Reflection,
diff --git a/Engine/Source/Developer/VulkanShaderFormat/Private/SpirVShaderCompiler.inl b/Engine/Source/Developer/VulkanShaderFormat/Private/SpirVShaderCompiler.inl
--- a/Engine/Source/Developer/VulkanShaderFormat/Private/SpirVShaderCompiler.inl
+++ b/Engine/Source/Developer/VulkanShaderFormat/Private/SpirVShaderCompiler.inl
@@ -974,7 +974,7 @@ static bool BuildShaderOutputFromSpirv(
 	SerializedOutput.Spirv.Data = TArray<uint32>(Reflection.GetCode(), Reflection.GetCodeSize() / 4);
 
 	// After we've copied the SPIRV module's copy of the SPIRV, use it to evaluate specialization constants
-	SerializedOutput.Header.SpecializationCount = UE::ShaderCompilerCommon::ProcessSpirvSpecializationConstants(InternalState.Input, CompilerContext, Reflection, Bindings, Output);
+	SerializedOutput.Header.SpecializationConstantData = UE::ShaderCompilerCommon::ProcessSpirvSpecializationConstants(InternalState.Input, CompilerContext, Reflection, Bindings, Output);
 
 	// We have to strip out most debug instructions (except OpName) for Vulkan mobile
 	if (InternalState.ShouldStripReflect())
diff --git a/Engine/Source/Developer/VulkanShaderFormat/Private/VulkanShaderCompiler.cpp b/Engine/Source/Developer/VulkanShaderFormat/Private/VulkanShaderCompiler.cpp
--- a/Engine/Source/Developer/VulkanShaderFormat/Private/VulkanShaderCompiler.cpp
+++ b/Engine/Source/Developer/VulkanShaderFormat/Private/VulkanShaderCompiler.cpp
@@ -371,7 +371,6 @@ void CompileVulkanShader(const FShaderCompilerInput& Input, const FShaderPreproc
 	}
 
 	Output.SerializeShaderCodeValidation();
-	Output.SerializeSpecializationConstantData();
 
 	ShaderParameterParser.ValidateShaderParameterTypes(Input, InternalState.IsMobileES31(), Output);
 	
diff --git a/Engine/Source/Runtime/RHICore/Private/RHICoreShader.cpp b/Engine/Source/Runtime/RHICore/Private/RHICoreShader.cpp
--- a/Engine/Source/Runtime/RHICore/Private/RHICoreShader.cpp
+++ b/Engine/Source/Runtime/RHICore/Private/RHICoreShader.cpp
@@ -79,24 +79,6 @@ void UE::RHICore::SetupShaderCodeValidationData(FRHIShader* RHIShader, FShaderCo
 #endif
 }
 
-void UE::RHICore::SetupSpecializationConstantData(FRHIShader* RHIShader, FShaderCodeReader& ShaderCodeReader)
-{
-#if PLATFORM_SUPPORTS_SPECIALIZATION_CONSTANTS
-	if (RHIShader)
-	{
-		int32 SpecializationConstantValidationSize = 0;
-		const uint8* SpecializationConstantValidationData = ShaderCodeReader.FindOptionalDataAndSize(EShaderOptionalDataKey::SpecializationConstant, SpecializationConstantValidationSize);
-		if (SpecializationConstantValidationData && SpecializationConstantValidationSize > 0)
-		{
-			FBufferReader ArValidationExtensionData((void*)SpecializationConstantValidationData, SpecializationConstantValidationSize, false);
-			FSpecializationConstantData SpecializationConstantData;
-			ArValidationExtensionData << SpecializationConstantData;
-			RHIShader->SetSpecializationExclusions(MoveTemp(SpecializationConstantData.SpecializationConstantExclusions));
-		}
-	}
-#endif
-}
-
 void UE::RHICore::DispatchShaderBundleEmulation(
 	FRHIComputeCommandList& InRHICmdList,
 	FRHIShaderBundle* ShaderBundle,
diff --git a/Engine/Source/Runtime/RHICore/Public/RHICoreShader.h b/Engine/Source/Runtime/RHICore/Public/RHICoreShader.h
--- a/Engine/Source/Runtime/RHICore/Public/RHICoreShader.h
+++ b/Engine/Source/Runtime/RHICore/Public/RHICoreShader.h
@@ -18,7 +18,6 @@ namespace RHICore
 /** Validates that the uniform buffer at the requested static slot. */
 extern RHICORE_API void ValidateStaticUniformBuffer(FRHIUniformBuffer* UniformBuffer, FUniformBufferStaticSlot Slot, uint32 ExpectedHash);
 extern RHICORE_API void SetupShaderCodeValidationData(FRHIShader* RHIShader, class FShaderCodeReader& ShaderCodeReader);
-extern RHICORE_API void SetupSpecializationConstantData(FRHIShader* RHIShader, class FShaderCodeReader& ShaderCodeReader);
 extern RHICORE_API void SetupShaderDiagnosticData(FRHIShader* RHIShader, class FShaderCodeReader& ShaderCodeReader);
 extern RHICORE_API void RegisterDiagnosticMessages(const TArray<FShaderDiagnosticData>& In);
 extern RHICORE_API const FString* GetDiagnosticMessage(uint32 MessageID);
diff --git a/Engine/Source/Runtime/RenderCore/Private/ShaderCore.cpp b/Engine/Source/Runtime/RenderCore/Private/ShaderCore.cpp
--- a/Engine/Source/Runtime/RenderCore/Private/ShaderCore.cpp
+++ b/Engine/Source/Runtime/RenderCore/Private/ShaderCore.cpp
@@ -1585,20 +1585,6 @@ void FShaderCompilerOutput::SerializeShaderCodeValidation()
 	}
 }
 
-void FShaderCompilerOutput::SerializeSpecializationConstantData()
-{
-	if (SpecializationConstantExclusions.Num() > 0)
-	{
-		FSpecializationConstantData SpecializationConstantData;
-		SpecializationConstantData.SpecializationConstantExclusions = SpecializationConstantExclusions;
-
-		TArray<uint8> WriterBytes;
-		FMemoryWriter Writer(WriterBytes);
-		Writer << SpecializationConstantData;
-		ShaderCode.AddOptionalData(FSpecializationConstantData::Key, WriterBytes.GetData(), WriterBytes.Num());
-	}
-}
-
 void FShaderCompilerOutput::SerializeShaderDiagnosticData()
 {
 	if (ShaderDiagnosticDatas.Num() > 0)
diff --git a/Engine/Source/Runtime/RenderCore/Public/ShaderCompilerCore.h b/Engine/Source/Runtime/RenderCore/Public/ShaderCompilerCore.h
--- a/Engine/Source/Runtime/RenderCore/Public/ShaderCompilerCore.h
+++ b/Engine/Source/Runtime/RenderCore/Public/ShaderCompilerCore.h
@@ -543,8 +543,6 @@ struct FShaderCompilerOutput
 	TArray<FShaderCodeValidationType> ParametersUAVTypeToValidate;
 	TArray<FShaderCodeValidationUBSize> ParametersUBSizeToValidate;
 
-	TMap<uint32, TArray<uint32>> SpecializationConstantExclusions;
-
 	TArray<FShaderDiagnosticData> ShaderDiagnosticDatas;
 
 	/** Use this field to store the shader source code if it's modified as part of the shader format's compilation process. This should only be set when 
@@ -570,9 +568,6 @@ struct FShaderCompilerOutput
 	/** Add optional data in ShaderCode to perform additional shader input validation at runtime*/
 	RENDERCORE_API void SerializeShaderCodeValidation();
 
-	/** Add specialization constant related data*/
-	RENDERCORE_API void SerializeSpecializationConstantData();
-
 	/** Add optional diagnostic data in ShaderCode to perform assert translation at runtime*/
 	RENDERCORE_API void SerializeShaderDiagnosticData();
 
diff --git a/Engine/Source/Runtime/RenderCore/Public/ShaderCore.h b/Engine/Source/Runtime/RenderCore/Public/ShaderCore.h
--- a/Engine/Source/Runtime/RenderCore/Public/ShaderCore.h
+++ b/Engine/Source/Runtime/RenderCore/Public/ShaderCore.h
@@ -922,24 +922,52 @@ struct FShaderCodeVendorExtension
 
 };
 
-
 struct FSpecializationConstantData
 {
-	// for FindOptionalData() and AddOptionalData()
-	static constexpr EShaderOptionalDataKey Key = EShaderOptionalDataKey::SpecializationConstant;
-	static constexpr uint16 StaticVersion = 0;
+	// Possible values for the specialization constants
+	TArray<TArray<int32>> Values;
 
-	TMap<uint32, TArray<uint32>> SpecializationConstantExclusions;
-	uint16 Version = StaticVersion;
+	// Pairs of SpecializationHash -> ArrayOfResourceExclusions
+	TArray<TPair<uint32, TArray<uint32>>> Exclusions;
 
-	friend FArchive& operator<<(FArchive& Ar, FSpecializationConstantData& Extension)
+	void Sort()
 	{
-		Ar << Extension.Version;
-		Ar << Extension.SpecializationConstantExclusions;
-		return Ar;
+		for (TArray<int32>& Array : Values)
+		{
+			Array.Sort();
+		}
+
+		for (TPair<uint32, TArray<uint32>>& Pair : Exclusions)
+		{
+			Pair.Value.Sort();
+		}
+		Exclusions.Sort([](const TPair<uint32, TArray<uint32>>& PairA, const TPair<uint32, TArray<uint32>>& PairB) { return PairA.Key < PairB.Key; });
+	}
+
+	TMap<uint32, TArray<uint32>> GetExclusionMap() const
+	{
+		TMap<uint32, TArray<uint32>> ExclusionMap;
+		for (const TPair<uint32, TArray<uint32>>& Pair : Exclusions)
+		{
+			ExclusionMap.Add(Pair);
+		}
+		return ExclusionMap;
+	}
+
+	void SetExclusionMap(TMap<uint32, TArray<uint32>>&& ExclusionMap)
+	{
+		for (TPair<uint32, TArray<uint32>>& Pair : ExclusionMap)
+		{
+			Exclusions.Emplace(Pair.Key, MoveTemp(Pair.Value));
+		}
 	}
 };
 
+inline FArchive& operator<<(FArchive& Ar, FSpecializationConstantData& SpecializationConstantData)
+{
+	return Ar << SpecializationConstantData.Values << SpecializationConstantData.Exclusions;
+}
+
 
 inline FArchive& operator<<(FArchive& Ar, FShaderCodeValidationStride& ShaderCodeValidationStride)
 {
diff --git a/Engine/Source/Runtime/VulkanRHI/Private/VulkanShaders.cpp b/Engine/Source/Runtime/VulkanRHI/Private/VulkanShaders.cpp
--- a/Engine/Source/Runtime/VulkanRHI/Private/VulkanShaders.cpp
+++ b/Engine/Source/Runtime/VulkanRHI/Private/VulkanShaders.cpp
@@ -55,8 +55,6 @@ static void ReadShaderOptionalData(FShaderCodeReader& ShaderCode, ShaderType* RH
 	RHIShader->Debug.ShaderName = ShaderCode.FindOptionalData(FShaderCodeName::Key);
 	UE::RHICore::SetupShaderCodeValidationData(RHIShader, ShaderCode);
 #endif
-
-	UE::RHICore::SetupSpecializationConstantData(RHIShader, ShaderCode);
 }
 
 static VkPipeline CreatePreRasterShaderLibrary(FVulkanDevice* Device, FVulkanShader* Shader, VkShaderStageFlagBits StageFlag)
@@ -290,7 +288,8 @@ ShaderType* FVulkanShaderFactory::CreateShader(TArrayView<const uint8> Code, FVu
 
 				if (FDataDrivenShaderPlatformInfo::GetSupportsSpecializationConstants(GMaxRHIShaderPlatform))
 				{
-					RetShader->SetSpecializationConstants(RetShader->GetCodeHeader().SpecializationCount > 0);
+					RetShader->SetSpecializationConstants(RetShader->GetCodeHeader().SpecializationConstantData.Values.Num() > 0);
+					RetShader->SetSpecializationExclusions(RetShader->GetCodeHeader().SpecializationConstantData.GetExclusionMap());
 				}
 
 				if (Device->SupportsShaderObjects())
diff --git a/Engine/Source/Runtime/VulkanRHI/Public/VulkanShaderResources.h b/Engine/Source/Runtime/VulkanRHI/Public/VulkanShaderResources.h
--- a/Engine/Source/Runtime/VulkanRHI/Public/VulkanShaderResources.h
+++ b/Engine/Source/Runtime/VulkanRHI/Public/VulkanShaderResources.h
@@ -104,7 +104,7 @@ struct FVulkanShaderHeader
 	uint32						RayTracingPayloadSize = 0;
 
 	// Specialization constant info
-	uint32						SpecializationCount = 0;
+	FSpecializationConstantData	SpecializationConstantData;
 
 	FShaderHash					CodeHash;
 	uint32						SpirvCRC = 0;
@@ -183,7 +183,7 @@ inline FArchive& operator<<(FArchive& Ar, FVulkanShaderHeader& Header)
 	Ar << Header.InOutMask;
 	Ar << Header.RayTracingPayloadType;
 	Ar << Header.RayTracingPayloadSize;
-	Ar << Header.SpecializationCount;
+	Ar << Header.SpecializationConstantData;
 	Ar << Header.CodeHash;
 	Ar << Header.SpirvCRC;
 	Ar << Header.WaveSize;
