Adding Push Constant support to FSpirvReflectBindings

#Synced-CL 52553930

diff --git a/Engine/Source/Developer/ShaderCompilerCommon/Private/SpirvReflectCommon.cpp b/Engine/Source/Developer/ShaderCompilerCommon/Private/SpirvReflectCommon.cpp
--- a/Engine/Source/Developer/ShaderCompilerCommon/Private/SpirvReflectCommon.cpp
+++ b/Engine/Source/Developer/ShaderCompilerCommon/Private/SpirvReflectCommon.cpp
@@ -170,6 +170,32 @@ void FSpirvReflectBindings::GatherOutputAttributes(const spv_reflect::ShaderModu
 	}
 }
 
+void FSpirvReflectBindings::GatherPushConstants(const spv_reflect::ShaderModule& SpirvReflection)
+{
+	// Enumerate all interface variables before filtering out push constants.
+	uint32 NumVariables = 0;
+	SpvReflectResult SpvResult = SpirvReflection.EnumerateEntryPointInterfaceVariables(SpirvReflection.GetEntryPointName(), &NumVariables, nullptr);
+	check(SpvResult == SPV_REFLECT_RESULT_SUCCESS);
+
+	if (NumVariables > 0)
+	{
+		TArray<SpvReflectInterfaceVariable*> AllVariables;
+		AllVariables.SetNum(NumVariables);
+
+		SpvResult = SpirvReflection.EnumerateEntryPointInterfaceVariables(SpirvReflection.GetEntryPointName(), &NumVariables, AllVariables.GetData());
+		check(SpvResult == SPV_REFLECT_RESULT_SUCCESS);
+
+		// Extract only the variables with the Push Constant storage class.
+		for (SpvReflectInterfaceVariable* Variable : AllVariables)
+		{
+			if (Variable && Variable->storage_class == SpvStorageClassPushConstant)
+			{
+				PushConstants.Emplace(Variable);
+			}
+		}
+	}
+}
+
 // Parse the index from a semantic name, e.g. "ATTRIBUTE14" returns 14.
 static bool ParseSemanticIndex(const ANSICHAR* InName, int32& OutSemanticIndex)
 {
diff --git a/Engine/Source/Developer/ShaderCompilerCommon/Public/SpirvReflectCommon.h b/Engine/Source/Developer/ShaderCompilerCommon/Public/SpirvReflectCommon.h
--- a/Engine/Source/Developer/ShaderCompilerCommon/Public/SpirvReflectCommon.h
+++ b/Engine/Source/Developer/ShaderCompilerCommon/Public/SpirvReflectCommon.h
@@ -20,6 +20,7 @@ struct FSpirvReflectBindings
 {
 	TArray<SpvReflectInterfaceVariable*> InputAttributes;
 	TArray<SpvReflectInterfaceVariable*> OutputAttributes;
+	TArray<SpvReflectInterfaceVariable*> PushConstants;
 	TSet<SpvReflectDescriptorBinding*> AtomicCounters;
 	TArray<SpvReflectDescriptorBinding*> InputAttachments; // for subpass inputs
 	TArray<SpvReflectDescriptorBinding*> UniformBuffers;
@@ -46,6 +47,9 @@ struct FSpirvReflectBindings
 	/** Gathers all output interface variables from the specified SPIRV-Reflect module. */
 	SHADERCOMPILERCOMMON_API void GatherOutputAttributes(const spv_reflect::ShaderModule& SpirvReflection);
 
+	/** Gathers all push constant variables from the specified SPIRV-Reflect module. */
+	SHADERCOMPILERCOMMON_API void GatherPushConstants(const spv_reflect::ShaderModule& SpirvReflection);
+
 	/** Assigns the binding location for all input attributes by its semantic index if their name is equal to 'SemanticName'. Used for vertex shaders where all input attributes have the "ATTRIBUTE" semantic. */
 	SHADERCOMPILERCOMMON_API void AssignInputAttributeLocationsBySemanticIndex(spv_reflect::ShaderModule& SpirvReflection, const ANSICHAR* SemanticName);
 };
